Console Variables Editor:

-Remove unnecessary include (needed change for external files to include Asset header)

nDisplay Launch:

-Initial Commit

diff --git a/Engine/Plugins/Editor/ConsoleVariablesEditor/Source/ConsoleVariablesEditor/Public/ConsoleVariablesAsset.h b/Engine/Plugins/Editor/ConsoleVariablesEditor/Source/ConsoleVariablesEditor/Public/ConsoleVariablesAsset.h
--- a/Engine/Plugins/Editor/ConsoleVariablesEditor/Source/ConsoleVariablesEditor/Public/ConsoleVariablesAsset.h
+++ b/Engine/Plugins/Editor/ConsoleVariablesEditor/Source/ConsoleVariablesEditor/Public/ConsoleVariablesAsset.h
@@ -2,8 +2,6 @@
 
 #pragma once
 
-#include "ConsoleVariablesEditorCommandInfo.h"
-
 #include "CoreMinimal.h"
 #include "Widgets/Input/SCheckBox.h"
 
diff --git a/Engine/Plugins/Editor/DisplayClusterLaunch/DisplayClusterLaunch.uplugin b/Engine/Plugins/Editor/DisplayClusterLaunch/DisplayClusterLaunch.uplugin
new file mode 100644
--- /dev/null
+++ b/Engine/Plugins/Editor/DisplayClusterLaunch/DisplayClusterLaunch.uplugin
@@ -0,0 +1,49 @@
+{
+	"FileVersion": 1,
+	"Version": 1,
+	"VersionName": "1.0",
+	"FriendlyName": "nDisplay Launch",
+	"Description": "Launch local nDisplay nodes with ease.",
+	"Category": "Virtual Production",
+	"CreatedBy": "",
+	"CreatedByURL": "",
+	"DocsURL": "",
+	"EnabledByDefault": false,
+	"MarketplaceURL": "",
+	"SupportURL": "",
+	"CanContainContent": true,
+	"IsBetaVersion": true,
+	"IsExperimentalVersion": false,
+	"Installed": false,
+	"Plugins":
+	[
+		{
+			"Name": "nDisplay",
+			"Enabled": true
+		},
+		{
+			"Name": "ConsoleVariables",
+			"Enabled": true
+		},
+		{
+			"Name": "ConcertMain",
+			"Enabled": true 
+		},
+		{
+			"Name": "MultiUserClient",
+			"Enabled": true
+		},
+		{
+			"Name": "UdpMessaging",
+			"Enabled": true
+		}
+	],
+	"Modules": [
+		{
+			"Name": "DisplayClusterLaunchEditor",
+			"Type": "Editor",
+			"LoadingPhase": "Default",
+			"WhitelistPlatforms": [ "Win64", "Linux" ]
+		}
+	]
+}
diff --git a/Engine/Plugins/Editor/DisplayClusterLaunch/Source/DisplayClusterLaunchEditor/DisplayClusterLaunchEditor.build.cs b/Engine/Plugins/Editor/DisplayClusterLaunch/Source/DisplayClusterLaunchEditor/DisplayClusterLaunchEditor.build.cs
new file mode 100644
--- /dev/null
+++ b/Engine/Plugins/Editor/DisplayClusterLaunch/Source/DisplayClusterLaunchEditor/DisplayClusterLaunchEditor.build.cs
@@ -0,0 +1,52 @@
+// Copyright Epic Games, Inc. All Rights Reserved.
+
+using System.IO;
+using UnrealBuildTool;
+
+public class DisplayClusterLaunchEditor : ModuleRules
+{
+	public DisplayClusterLaunchEditor(ReadOnlyTargetRules Target) : base(Target)
+	{
+		PCHUsage = ModuleRules.PCHUsageMode.UseExplicitOrSharedPCHs;
+
+		PublicDependencyModuleNames.AddRange(
+			new string[]
+			{
+				"Core",
+				"OutputLog"
+			}
+			);
+			
+		PrivateDependencyModuleNames.AddRange(
+			new string[]
+			{
+				"AssetRegistry",
+				"AssetTools",
+				"Concert",
+				"ConcertTransport",
+				"CoreUObject",
+				"ConsoleVariablesEditor",
+				"ContentBrowser",
+				"DisplayCluster",
+				"DisplayClusterConfiguration",
+				"Engine",
+				"EditorStyle",
+				"EditorWidgets",
+				"InputCore",
+				"Kismet",
+				"Messaging",
+				"MultiUserClient",
+				"Networking",
+				"Projects",
+				"PropertyEditor",
+				"Slate",
+				"SlateCore",
+				"ToolMenus", 
+				"ToolWidgets",
+				"UdpMessaging",
+				"UnrealEd",
+				"WorkspaceMenuStructure"
+			}
+			);
+	}
+}
diff --git a/Engine/Plugins/Editor/DisplayClusterLaunch/Source/DisplayClusterLaunchEditor/Private/DisplayClusterLaunchEditorModule.cpp b/Engine/Plugins/Editor/DisplayClusterLaunch/Source/DisplayClusterLaunchEditor/Private/DisplayClusterLaunchEditorModule.cpp
new file mode 100644
--- /dev/null
+++ b/Engine/Plugins/Editor/DisplayClusterLaunch/Source/DisplayClusterLaunchEditor/Private/DisplayClusterLaunchEditorModule.cpp
@@ -0,0 +1,726 @@
+// Copyright Epic Games, Inc. All Rights Reserved.
+
+#include "DisplayClusterLaunchEditorModule.h"
+
+#include "DisplayClusterLaunchEditorLog.h"
+#include "DisplayClusterLaunchEditorProjectSettings.h"
+#include "DisplayClusterLaunchEditorStyle.h"
+
+#include "ConcertSettings.h"
+#include "IMultiUserClientModule.h"
+
+#include "DisplayClusterRootActor.h"
+#include "IDisplayClusterConfiguration.h"
+#include "DisplayClusterConfigurationTypes.h"
+
+#include "Shared/UdpMessagingSettings.h"
+
+#include "AssetRegistry/IAssetRegistry.h"
+#include "Editor.h"
+#include "EngineUtils.h"
+#include "Engine/GameEngine.h"
+#include "Framework/Commands/Commands.h"
+#include "ISettingsModule.h"
+#include "LevelEditor.h"
+#include "Misc/ConfigCacheIni.h"
+#include "ToolMenus.h"
+
+#define LOCTEXT_NAMESPACE "FDisplayClusterLaunchEditorModule"
+
+FString EnumToString(const TCHAR* EnumName, const int32 EnumValue)
+{
+	const UEnum* EnumPtr = FindObject<UEnum>(ANY_PACKAGE, EnumName, true);
+	
+	if (!EnumPtr)
+	{
+		return LOCTEXT("EnumNotFound", "Enum not found").ToString();
+	}
+
+	return EnumPtr->GetNameStringByIndex(EnumValue);
+}
+
+static UWorld* GetCurrentWorld()
+{
+	UWorld* CurrentWorld = nullptr;
+	if (GIsEditor)
+	{
+		CurrentWorld = GEditor->GetEditorWorldContext().World();
+	}
+	else if (UGameEngine* GameEngine = Cast<UGameEngine>(GEngine))
+	{
+		CurrentWorld = GameEngine->GetGameWorld();
+	}
+	return CurrentWorld;
+}
+
+class FDisplayClusterLaunchUiCommands : public TCommands<FDisplayClusterLaunchUiCommands>
+{
+public:
+	FDisplayClusterLaunchUiCommands()
+		: TCommands<FDisplayClusterLaunchUiCommands>(
+			"DisplayClusterLaunch",
+			LOCTEXT("DisplayClusterLaunchCommands", "DisplayClusterLaunch UI Commands"),
+			NAME_None,
+			FDisplayClusterLaunchEditorStyle::Get().GetStyleSetName()
+		)
+	{}
+
+	virtual void RegisterCommands() override
+	{
+		UI_COMMAND(LaunchDisplayCluster, "Launch nDisplay", "Launch nDisplay", EUserInterfaceActionType::Button, FInputChord());
+	}
+
+	TSharedPtr<FUICommandInfo> LaunchDisplayCluster;
+};
+
+FDisplayClusterLaunchEditorModule& FDisplayClusterLaunchEditorModule::Get()
+{
+	return FModuleManager::GetModuleChecked<FDisplayClusterLaunchEditorModule>("DisplayClusterLaunchEditor");
+}
+
+void FDisplayClusterLaunchEditorModule::StartupModule()
+{
+	FDisplayClusterLaunchEditorStyle::Initialize();
+
+	FCoreDelegates::OnFEngineLoopInitComplete.AddRaw(this, &FDisplayClusterLaunchEditorModule::OnFEngineLoopInitComplete);
+}
+
+void FDisplayClusterLaunchEditorModule::ShutdownModule()
+{
+	UToolMenus::UnregisterOwner(this);
+
+	FCoreDelegates::OnFEngineLoopInitComplete.RemoveAll(this);
+
+	FDisplayClusterLaunchEditorStyle::Shutdown();
+
+	// Unregister project settings
+	ISettingsModule& SettingsModule = FModuleManager::LoadModuleChecked<ISettingsModule>("Settings");
+	{
+		SettingsModule.UnregisterSettings("Project", "Plugins", "nDisplay Launch");
+	}
+}
+
+void FDisplayClusterLaunchEditorModule::OpenProjectSettings()
+{
+	FModuleManager::LoadModuleChecked<ISettingsModule>("Settings")
+		.ShowViewer("Project", "Plugins", "nDisplay Launch");
+
+}
+
+void GetProjectSettingsArguments(const UDisplayClusterLaunchEditorProjectSettings* ProjectSettings, FString& ConcatenatedCommandLineArguments, FString& ConcatenatedConsoleCommands, FString& ConcatenatedDPCvars, FString& ConcatenatedLogCommands)
+{
+	{
+		for (const FString& CommandLineArgument : ProjectSettings->CommandLineArguments)
+		{
+			if (CommandLineArgument.IsEmpty())
+			{
+				continue;
+			}
+			ConcatenatedCommandLineArguments += FString::Printf(TEXT("-%s "), *CommandLineArgument);
+		}
+		// Remove whitespace
+		ConcatenatedCommandLineArguments.TrimStartAndEndInline();
+	}
+
+	if (ProjectSettings->ConsoleCommands.Num() > 0)
+	{
+		ConcatenatedConsoleCommands += FString::Join(ProjectSettings->ConsoleCommands, TEXT(","));
+	}
+
+	if (ProjectSettings->DPCvars.Num() > 0)
+	{
+		ConcatenatedDPCvars += FString::Join(ProjectSettings->DPCvars, TEXT(","));
+	}
+
+	{
+		for (const FDisplayClusterLaunchLoggingConstruct& LoggingConstruct : ProjectSettings->Logging)
+		{
+			if (LoggingConstruct.Category.IsNone())
+			{
+				continue;
+			}
+			ConcatenatedLogCommands += FString::Printf(TEXT("%s %s, "),
+			                                           *LoggingConstruct.Category.ToString(),
+			                                           *EnumToString(TEXT("EDisplayClusterLaunchLogVerbosity"), (int32)LoggingConstruct.VerbosityLevel.GetValue()));
+		}
+		// Remove whitespace
+		ConcatenatedLogCommands.TrimStartAndEndInline();
+		// Remove last comma
+		ConcatenatedLogCommands = ConcatenatedLogCommands.LeftChop(1);
+	}
+}
+
+bool AddUdpMessagingArguments(FString& ConcatenatedArguments)
+{
+	// Get from reflection because UUdpMessagingSettings does not export its API
+	FConfigFile* EngineConfig = GConfig ? GConfig->FindConfigFileWithBaseName(FName(TEXT("Engine"))) : nullptr;
+	if (EngineConfig)
+	{
+		TArray<FString> Settings;
+		FString Setting;
+
+		// Unicast endpoint setting
+		EngineConfig->GetString(TEXT("/Script/UdpMessaging.UdpMessagingSettings"), TEXT("UnicastEndpoint"), Setting);
+
+		// if the unicast endpoint port is bound, concatenate it
+		if (Setting != "0.0.0.0:0" && !Setting.IsEmpty())
+		{
+			ConcatenatedArguments += TEXT(" -UDPMESSAGING_TRANSPORT_UNICAST=") + Setting;
+		}
+
+		// Multicast endpoint setting
+		EngineConfig->GetString(TEXT("/Script/UdpMessaging.UdpMessagingSettings"), TEXT("MulticastEndpoint"), Setting);
+		ConcatenatedArguments += TEXT(" -UDPMESSAGING_TRANSPORT_MULTICAST=") + Setting;
+
+		// Static endpoints setting
+		Settings.Empty(1);
+		EngineConfig->GetArray(TEXT("/Script/UdpMessaging.UdpMessagingSettings"), TEXT("StaticEndpoints"), Settings);
+		if (Settings.Num() > 0)
+		{
+			ConcatenatedArguments += TEXT(" -UDPMESSAGING_TRANSPORT_STATIC=");
+			ConcatenatedArguments += Settings[0];
+			for (int32 i = 1; i < Settings.Num(); ++i)
+			{
+				ConcatenatedArguments += ',';
+				ConcatenatedArguments += Settings[i];
+			}
+		}
+		return true;
+	}
+	return false;
+}
+
+FString GetConcertArguments()
+{
+	const UConcertClientConfig* ConcertClientConfig = GetDefault<UConcertClientConfig>();
+	ensureAlwaysMsgf (ConcertClientConfig, TEXT("%hs: Unable to launch nDisplay because there is no UConcertClientConfig object."));
+
+	FString ReturnValue =
+		FString::Printf(TEXT("-CONCERTISHEADLESS -CONCERTRETRYAUTOCONNECTONERROR -CONCERTAUTOCONNECT -CONCERTSERVER=\"%s\" -CONCERTSESSION=\"%s\""),
+			*ConcertClientConfig->DefaultServerURL, *ConcertClientConfig->DefaultSessionName);
+
+	return ReturnValue;
+}
+
+void LaunchConcertServerIfNotRunning()
+{
+	IMultiUserClientModule& MultiUserClientModule = IMultiUserClientModule::Get();
+	if (!MultiUserClientModule.IsConcertServerRunning())
+	{
+		MultiUserClientModule.LaunchConcertServer();
+	}
+}
+
+void FDisplayClusterLaunchEditorModule::LaunchDisplayClusterProcess()
+{
+	if (!SelectedDisplayClusterConfigActor.IsValid())
+	{
+		UE_LOG(LogDisplayClusterLaunchEditor, Error, TEXT("%hs: Unable to launch nDisplay because there is no valid selected nDisplay configuration."), __FUNCTION__);
+		return;
+	}
+
+	if (!DoesCurrentWorldHaveDisplayClusterConfig())
+	{
+		UE_LOG(LogDisplayClusterLaunchEditor, Error, TEXT("%hs: Unable to launch nDisplay because there are no valid selected nDisplay configurations in the world."), __FUNCTION__);
+		return;
+	}
+
+	if (SelectedDisplayClusterConfigActorNodes.Num() == 0)
+	{
+		UE_LOG(LogDisplayClusterLaunchEditor, Error, TEXT("%hs: Unable to launch nDisplay because there are no selected nDisplay nodes."), __FUNCTION__);
+		return;
+	}
+
+	FString ConfigActorPath;
+	if (const ADisplayClusterRootActor* ConfigActor = Cast<ADisplayClusterRootActor>(SelectedDisplayClusterConfigActor.ResolveObject()))
+	{
+		const FString FilePath = FPaths::Combine(FPaths::ProjectSavedDir(), "Temp.ndisplay");
+		UDisplayClusterConfigurationData* ConfigDataCopy = DuplicateObject(ConfigActor->GetConfigData(), GetTransientPackage());
+
+		if (!ConfigDataCopy->Scene)
+		{
+			ConfigDataCopy->Scene = NewObject<UDisplayClusterConfigurationScene>(ConfigDataCopy, NAME_None,
+			RF_ArchetypeObject | RF_Public );
+		}
+		
+		if (!ensureAlways(IDisplayClusterConfiguration::Get().SaveConfig(ConfigDataCopy, FilePath)))
+		{
+			UE_LOG(LogDisplayClusterLaunchEditor, Error, TEXT("%hs: Unable to launch nDisplay because the selected nDisplay Configuration could not be saved to a .ndisplay file. See the log for more information."), __FUNCTION__);
+			return;
+		}
+		
+		ConfigActorPath = FString::Printf(TEXT("-dc_cfg=\"%s\""), *FilePath);
+	}
+	else
+	{
+		UE_LOG(LogDisplayClusterLaunchEditor, Error, TEXT("%hs: Unable to launch nDisplay because the selected nDisplay Config Actor could not be resolved or does not exist in the current level."), __FUNCTION__);
+		return;
+	}
+	
+	const UDisplayClusterLaunchEditorProjectSettings* ProjectSettings = GetDefault<UDisplayClusterLaunchEditorProjectSettings>();
+	if (!ensureAlwaysMsgf (ProjectSettings, TEXT("%hs: Unable to launch nDisplay because there is no UDisplayClusterLaunchEditorProjectSettings object.")))
+	{
+		return;
+	}
+	
+	const FString EditorBinary = FPlatformProcess::ExecutablePath();
+	
+	const FString Project = FPaths::SetExtension(FPaths::Combine(FPaths::ProjectDir(), FApp::GetProjectName()),".uproject");
+	const FString Map = GetCurrentWorld()->GetCurrentLevel()->GetPackage()->GetLoadedPath().GetLocalFullPath();
+
+	// Create Multi-user params
+	FString ConcertArguments;
+	if (GetConnectToMultiUser())
+	{
+		GetConcertArguments();
+		LaunchConcertServerIfNotRunning();
+	}
+
+	for (const FString& Node : SelectedDisplayClusterConfigActorNodes)
+	{
+		FString ConcatenatedCommandLineArguments;
+		FString ConcatenatedConsoleCommands;
+		FString ConcatenatedDPCvars;
+		FString ConcatenatedLogCommands;
+	
+		GetProjectSettingsArguments(ProjectSettings, ConcatenatedCommandLineArguments, ConcatenatedConsoleCommands, ConcatenatedDPCvars,
+					 ConcatenatedLogCommands);
+
+		AddUdpMessagingArguments(ConcatenatedCommandLineArguments);
+
+		// Add nDisplay node information
+		ConcatenatedCommandLineArguments += " " + ConfigActorPath;
+		ConcatenatedCommandLineArguments += FString::Printf(TEXT(" -dc_node=\"%s\""), *Node);
+
+		// Add Multi-User params
+		if (!ConcertArguments.IsEmpty())
+		{
+			ConcatenatedCommandLineArguments += " " + ConcertArguments;
+		}
+
+		// Log file
+		const FString LogFileName = (ProjectSettings->LogFileName.IsEmpty() ? Node : ProjectSettings->LogFileName) + ".log";
+
+		const FString Params =
+			FString::Printf(
+				TEXT("\"%s\" -game \"%s\" Log=%s %s -ExecCmds=\"%s\" -DPCVars=\"%s\" -LogCmds=\"%s\""),
+				*Project, *Map, *LogFileName, *ConcatenatedCommandLineArguments,
+				*ConcatenatedConsoleCommands, *ConcatenatedDPCvars, *ConcatenatedLogCommands
+	
+		);
+
+		UE_LOG(LogDisplayClusterLaunchEditor, Log, TEXT("Full Command: %s %s"), *EditorBinary, *Params);
+
+		void* WritePipe = nullptr;
+
+		FProcHandle Process = FPlatformProcess::CreateProc(
+			*EditorBinary, *Params,
+			ProjectSettings->bCloseEditorOnLaunch,
+			false, false, NULL,
+			0, NULL, WritePipe
+		);
+	}
+}
+
+void FDisplayClusterLaunchEditorModule::OnFEngineLoopInitComplete()
+{
+	Actions = MakeShared<FUICommandList>();
+	FDisplayClusterLaunchUiCommands::Register();
+
+	Actions->MapAction(FDisplayClusterLaunchUiCommands::Get().LaunchDisplayCluster,
+		FExecuteAction::CreateRaw(this, &FDisplayClusterLaunchEditorModule::LaunchDisplayClusterProcess));
+	
+	RegisterToolbarItem();
+	RegisterProjectSettings();
+}
+
+void FDisplayClusterLaunchEditorModule::RegisterToolbarItem()
+{
+	UToolMenu* Menu = UToolMenus::Get()->ExtendMenu("LevelEditor.LevelEditorToolBar.User");
+	
+	RemoveToolbarItem();
+
+	FToolMenuSection& Section = Menu->AddSection("DisplayClusterLaunch");
+
+	FToolMenuEntry DisplayClusterLaunchButton = FToolMenuEntry::InitToolBarButton(FDisplayClusterLaunchUiCommands::Get().LaunchDisplayCluster);
+	DisplayClusterLaunchButton.SetCommandList(Actions);
+
+	const FToolMenuEntry DisplayClusterLaunchComboButton = FToolMenuEntry::InitComboButton(
+		"DisplayClusterLaunchMenu",
+		FUIAction(),
+		FOnGetContent::CreateRaw(this, &FDisplayClusterLaunchEditorModule::CreateToolbarMenuEntries),
+		LOCTEXT("DisplayClusterLaunchActions", "Launch nDisplay Actions"),
+		LOCTEXT("DisplayClusterLaunchActionsTooltip", "Actions related to nDisplay Launch"),
+		FSlateIcon(),
+		true //bInSimpleComboBox
+	);
+
+	Section.AddEntry(DisplayClusterLaunchButton);
+	Section.AddEntry(DisplayClusterLaunchComboButton);
+}
+
+void FDisplayClusterLaunchEditorModule::RemoveToolbarItem()
+{
+	UToolMenu* Menu = UToolMenus::Get()->ExtendMenu("LevelEditor.LevelEditorToolBar.User");
+	if (Menu->FindSection("DisplayClusterLaunch"))
+	{
+		Menu->RemoveSection("DisplayClusterLaunch");
+	}
+}
+
+void FDisplayClusterLaunchEditorModule::RegisterProjectSettings() const
+{
+	ISettingsModule& SettingsModule = FModuleManager::LoadModuleChecked<ISettingsModule>("Settings");
+	{
+		// User Project Settings
+		const TSharedPtr<ISettingsSection> ProjectSettingsSectionPtr = SettingsModule.RegisterSettings(
+			"Project", "Plugins", "nDisplay Launch",
+			LOCTEXT("DisplayClusterLaunchSettingsCategoryDisplayName", "nDisplay Launch"),
+			LOCTEXT("DisplayClusterLaunchSettingsDescription", "Configure the nDisplay Launch user settings"),
+			GetMutableDefault<UDisplayClusterLaunchEditorProjectSettings>());
+	}
+}
+
+TArray<TWeakObjectPtr<ADisplayClusterRootActor>> FDisplayClusterLaunchEditorModule::GetAllDisplayClusterConfigsInWorld()
+{
+	TArray<TWeakObjectPtr<ADisplayClusterRootActor>> CachedDisplayClusterActors;
+	for (TActorIterator<ADisplayClusterRootActor> Itr(GetCurrentWorld()); Itr; ++Itr)
+	{
+		if (ADisplayClusterRootActor* RootActor = *Itr)
+		{
+			CachedDisplayClusterActors.Add(RootActor);
+		}
+	}
+
+	bAreConfigsFoundInWorld = CachedDisplayClusterActors.Num() > 0;
+
+	CachedDisplayClusterActors.Sort([](const TWeakObjectPtr<ADisplayClusterRootActor>& A, const TWeakObjectPtr<ADisplayClusterRootActor>& B)
+	{
+		return A->GetActorLabel() < B->GetActorLabel();
+	});
+
+	return CachedDisplayClusterActors;
+}
+
+bool FDisplayClusterLaunchEditorModule::DoesCurrentWorldHaveDisplayClusterConfig()
+{
+	return bAreConfigsFoundInWorld;
+}
+
+void FDisplayClusterLaunchEditorModule::SetSelectedDisplayClusterConfigActor(ADisplayClusterRootActor* SelectedActor)
+{
+	if (SelectedActor)
+	{
+		const FSoftObjectPath AsSoftObjectPath(SelectedActor);
+
+		if (AsSoftObjectPath != SelectedDisplayClusterConfigActor)
+		{
+			SelectedDisplayClusterConfigActor = AsSoftObjectPath;
+
+			SelectedDisplayClusterConfigActorNodes.Empty();
+		}
+	}
+}
+
+void FDisplayClusterLaunchEditorModule::ToggleDisplayClusterConfigActorNodeSelected(FString InNodeName)
+{
+	if (IsDisplayClusterConfigActorNodeSelected(InNodeName))
+	{
+		SelectedDisplayClusterConfigActorNodes.Remove(InNodeName);
+	}
+	else
+	{
+		SelectedDisplayClusterConfigActorNodes.Add(InNodeName);
+	}
+}
+
+bool FDisplayClusterLaunchEditorModule::IsDisplayClusterConfigActorNodeSelected(FString InNodeName)
+{
+	return SelectedDisplayClusterConfigActorNodes.Contains(InNodeName);
+}
+
+void FDisplayClusterLaunchEditorModule::SetSelectedConsoleVariablesAsset(const FAssetData InConsoleVariablesAsset)
+{
+	if (SelectedConsoleVariablesAssetName == InConsoleVariablesAsset.AssetName)
+	{
+		SelectedConsoleVariablesAssetName = NAME_None;
+	}
+	else
+	{
+		SelectedConsoleVariablesAssetName = InConsoleVariablesAsset.AssetName;
+	}
+}
+
+TSharedRef<SWidget>  FDisplayClusterLaunchEditorModule::CreateToolbarMenuEntries()
+{
+	IAssetRegistry* AssetRegistry = IAssetRegistry::Get();
+	
+	FMenuBuilder MenuBuilder(true, Actions);
+
+	TArray<TWeakObjectPtr<ADisplayClusterRootActor>> DisplayClusterConfigs = GetAllDisplayClusterConfigsInWorld();
+
+	MenuBuilder.BeginSection("DisplayClusterLaunch", LOCTEXT("DisplayClusterLauncher", "Launch nDisplay"));
+	{
+		MenuBuilder.AddMenuEntry(
+			LOCTEXT("DisplayClusterLaunchLastNode", "Launch Last Node"),
+			LOCTEXT("DisplayClusterLaunchLastNodeTooltip", "Launch the last node."),
+			FSlateIcon(),
+			FUIAction(
+				FExecuteAction::CreateRaw(this, &FDisplayClusterLaunchEditorModule::LaunchDisplayClusterProcess),
+				FCanExecuteAction::CreateRaw(this, &FDisplayClusterLaunchEditorModule::DoesCurrentWorldHaveDisplayClusterConfig)
+			),
+			NAME_None
+		);
+	}
+	MenuBuilder.EndSection();
+
+	AddDisplayClusterLaunchConfigurations(AssetRegistry, MenuBuilder, DisplayClusterConfigs);
+
+	AddDisplayClusterLaunchNodes(AssetRegistry, MenuBuilder);
+
+	AddConsoleVariablesEditorAssetsToToolbarMenu(AssetRegistry, MenuBuilder);
+
+	AddOptionsToToolbarMenu(MenuBuilder);
+
+	return MenuBuilder.MakeWidget();
+}
+
+void FDisplayClusterLaunchEditorModule::AddDisplayClusterLaunchConfigurations(IAssetRegistry* AssetRegistry,
+	FMenuBuilder& MenuBuilder, TArray<TWeakObjectPtr<ADisplayClusterRootActor>>& DisplayClusterConfigs)
+{
+	MenuBuilder.BeginSection("DisplayClusterLaunchConfigurations", LOCTEXT("DisplayClusterLaunchConfigurations", "Configuration"));
+	{
+		if (DisplayClusterConfigs.Num())
+		{
+			// If one is not set, select the first one found
+			bool bIsConfigActorValid = false;
+			if (ADisplayClusterRootActor* SelectedActor = Cast<ADisplayClusterRootActor>(SelectedDisplayClusterConfigActor.ResolveObject()))
+			{
+				bIsConfigActorValid = DisplayClusterConfigs.ContainsByPredicate([SelectedActor](TWeakObjectPtr<ADisplayClusterRootActor>& Comparator)
+				{
+					return Comparator.IsValid() && SelectedActor == Comparator.Get();
+				});
+			}
+
+			if (!bIsConfigActorValid)
+			{
+				SetSelectedDisplayClusterConfigActor(DisplayClusterConfigs[0].Get());
+			}
+			
+			for (const TWeakObjectPtr<ADisplayClusterRootActor>& Node : DisplayClusterConfigs)
+			{
+				if (!Node.IsValid())
+				{
+					continue;
+				}
+				
+				const FText NodeName = FText::FromString(Node->GetActorLabel());
+				const FText DisplayClusterConfigTooltip =
+					FText::Format(LOCTEXT("SelectDisplayClusterConfigFormat","Select configuration '{0}'"), NodeName);
+			
+				MenuBuilder.AddMenuEntry(
+					NodeName,
+					DisplayClusterConfigTooltip,
+					FSlateIcon(),
+					FUIAction(
+						FExecuteAction::CreateRaw(this, &FDisplayClusterLaunchEditorModule::SetSelectedDisplayClusterConfigActor, Node.Get()),
+						FCanExecuteAction::CreateRaw(this, &FDisplayClusterLaunchEditorModule::DoesCurrentWorldHaveDisplayClusterConfig),
+						FIsActionChecked::CreateLambda([this, Node]() { return SelectedDisplayClusterConfigActor == FSoftObjectPath(Node.Get()); })
+					),
+					NAME_None,
+					EUserInterfaceActionType::RadioButton
+				);
+			}
+		}
+		else
+		{
+			SelectedDisplayClusterConfigActor.Reset();
+			
+			MenuBuilder.AddWidget(SNullWidget::NullWidget,
+				LOCTEXT("NoDisplayClusterConfigAssetInLevelText", "Add an nDisplay configuration asset to\nthe current level to launch nDisplay."));
+		}
+	}
+	MenuBuilder.EndSection();
+}
+
+void FDisplayClusterLaunchEditorModule::AddDisplayClusterLaunchNodes(IAssetRegistry* AssetRegistry, FMenuBuilder& MenuBuilder)
+{
+	if (ADisplayClusterRootActor* SelectedActor = Cast<ADisplayClusterRootActor>(SelectedDisplayClusterConfigActor.ResolveObject()))
+	{
+		TArray<FString> NodeNames;
+		SelectedActor->GetConfigData()->Cluster->Nodes.GenerateKeyArray(NodeNames);
+
+		if (NodeNames.Num() > 0)
+		{
+			MenuBuilder.BeginSection("DisplayClusterLaunchNodes", LOCTEXT("DisplayClusterLaunchNodes", "Nodes"));
+			{
+				const TAttribute<FText> SelectedDisplayClusterNodeNames =
+					TAttribute<FText>::Create(
+						[this]()
+						{
+							if (SelectedDisplayClusterConfigActorNodes.Num() > 0)
+							{
+								return FText::FromString(FString::Join(SelectedDisplayClusterConfigActorNodes, TEXT(", ")));
+							}
+							return LOCTEXT("NoDisplayClusterLaunchNodesSelected", "Please select nDisplay nodes to launch.");
+						}
+					);
+				
+				MenuBuilder.AddSubMenu(
+					SelectedDisplayClusterNodeNames,
+					LOCTEXT("SelectDisplayClusterNodes","Select nDisplay Nodes"),
+					FNewMenuDelegate::CreateLambda([this, NodeNames](FMenuBuilder& NewMenuBuilder)
+					{
+						for (const FString& NodeName : NodeNames)
+						{
+							const FText DisplayClusterNodeName = FText::FromString(NodeName);
+							const FText DisplayClusterNodeTooltip =
+								FText::Format(LOCTEXT("SelectDisplayClusterNodeFormat","Select node '{0}'"), DisplayClusterNodeName);
+							
+							NewMenuBuilder.AddMenuEntry(
+								DisplayClusterNodeName,
+								DisplayClusterNodeTooltip,
+								FSlateIcon(),
+								FUIAction(
+									FExecuteAction::CreateRaw(
+										this, &FDisplayClusterLaunchEditorModule::ToggleDisplayClusterConfigActorNodeSelected, NodeName),
+									FCanExecuteAction(),
+									FIsActionChecked::CreateRaw(
+										this, &FDisplayClusterLaunchEditorModule::IsDisplayClusterConfigActorNodeSelected, NodeName)
+								),
+								NAME_None,
+								EUserInterfaceActionType::Check
+							);
+						}
+					}),
+					FUIAction(FExecuteAction(), FCanExecuteAction::CreateRaw(this, &FDisplayClusterLaunchEditorModule::DoesCurrentWorldHaveDisplayClusterConfig)),
+					NAME_None,
+					EUserInterfaceActionType::None
+				);
+			}
+			MenuBuilder.EndSection();
+		}
+	}
+}
+
+void FDisplayClusterLaunchEditorModule::AddConsoleVariablesEditorAssetsToToolbarMenu(
+	IAssetRegistry* AssetRegistry, FMenuBuilder& MenuBuilder)
+{
+	TArray<FAssetData> FoundConsoleVariablesAssets;
+	AssetRegistry->GetAssetsByClass("ConsoleVariablesAsset", FoundConsoleVariablesAssets, true);
+
+	if (FoundConsoleVariablesAssets.Num())
+	{
+		MenuBuilder.BeginSection("DisplayClusterLaunchCvars", LOCTEXT("DisplayClusterLaunchCvars", "Console Variables"));
+		{
+			const FText ConsoleVariablesAssetTooltip = LOCTEXT("SelectConsoleVariablesAssetFormat","Select Console Variables Asset");
+			
+			MenuBuilder.AddSubMenu(
+				TAttribute<FText>::Create([this](){ return FText::FromName(SelectedConsoleVariablesAssetName); }),
+				ConsoleVariablesAssetTooltip,
+				FNewMenuDelegate::CreateLambda([this, FoundConsoleVariablesAssets](FMenuBuilder& NewMenuBuilder)
+				{
+					for (const FAssetData& Asset : FoundConsoleVariablesAssets)
+					{
+						const FText ConsoleVariablesAssetName = FText::FromName(Asset.AssetName);
+						const FText ConsoleVariablesAssetTooltip =
+							FText::Format(LOCTEXT("SelectConsoleVariablesAssetFormat","Select Console Variables Asset '{0}'"), ConsoleVariablesAssetName);
+						NewMenuBuilder.AddMenuEntry(
+							ConsoleVariablesAssetName,
+							ConsoleVariablesAssetTooltip,
+							FSlateIcon(),
+							FUIAction(
+								FExecuteAction::CreateRaw(this, &FDisplayClusterLaunchEditorModule::SetSelectedConsoleVariablesAsset, Asset),
+								FCanExecuteAction(),
+								FGetActionCheckState::CreateLambda([this, Asset](){ return SelectedConsoleVariablesAssetName == Asset.AssetName ? ECheckBoxState::Checked : ECheckBoxState::Unchecked; })
+							),
+							NAME_None,
+							EUserInterfaceActionType::RadioButton
+						);
+					}
+				}),
+				FUIAction(FExecuteAction(), FCanExecuteAction::CreateRaw(this, &FDisplayClusterLaunchEditorModule::DoesCurrentWorldHaveDisplayClusterConfig)),
+				NAME_None,
+				EUserInterfaceActionType::None
+			);
+		}
+		MenuBuilder.EndSection();
+	}
+}
+
+void FDisplayClusterLaunchEditorModule::AddOptionsToToolbarMenu(FMenuBuilder& MenuBuilder)
+{
+	MenuBuilder.BeginSection("DisplayClusterLaunchOptions", LOCTEXT("DisplayClusterLaunchOptions", "Options"));
+	{
+		MenuBuilder.AddMenuEntry(
+			LOCTEXT("ConnectToMultiUserLabel", "Connect to Multi-User"),
+			LOCTEXT("ConnectToMultiUserTooltip", "Connect to Multi-User"),
+			FSlateIcon(),
+			FUIAction(
+				FExecuteAction::CreateLambda([this]()
+				{
+					UDisplayClusterLaunchEditorProjectSettings* Settings = GetMutableDefault<UDisplayClusterLaunchEditorProjectSettings>();
+					Settings->bConnectToMultiUser = !GetConnectToMultiUser();
+				}),
+				FCanExecuteAction(),
+				FIsActionChecked::CreateRaw(this, &FDisplayClusterLaunchEditorModule::GetConnectToMultiUser)
+			),
+			NAME_None,
+			EUserInterfaceActionType::Check
+		);
+
+		MenuBuilder.AddMenuEntry(
+			LOCTEXT("EnableUnrealInsightsLabel", "Enable Unreal Insights"),
+			LOCTEXT("EnableUnrealInsightsTooltip", "Enable Unreal Insights"),
+			FSlateIcon(),
+			FUIAction(
+				FExecuteAction::CreateLambda([this]()
+				{
+					UDisplayClusterLaunchEditorProjectSettings* Settings = GetMutableDefault<UDisplayClusterLaunchEditorProjectSettings>();
+					Settings->bEnableUnrealInsights = !Settings->bEnableUnrealInsights;
+				}),
+				FCanExecuteAction(),
+				FIsActionChecked::CreateLambda([this](){ return GetDefault<UDisplayClusterLaunchEditorProjectSettings>()->bEnableUnrealInsights; })
+			),
+			NAME_None,
+			EUserInterfaceActionType::Check
+		);
+
+		MenuBuilder.AddMenuEntry(
+			LOCTEXT("CloseEditorOnLaunchLabel", "Close Editor on Launch"),
+			LOCTEXT("CloseEditorOnLaunchTooltip", "Close Editor on Launch"),
+			FSlateIcon(),
+			FUIAction(
+				FExecuteAction::CreateLambda([this]()
+				{
+					UDisplayClusterLaunchEditorProjectSettings* Settings = GetMutableDefault<UDisplayClusterLaunchEditorProjectSettings>();
+					Settings->bCloseEditorOnLaunch = !Settings->bCloseEditorOnLaunch;
+				}),
+				FCanExecuteAction(),
+				FIsActionChecked::CreateLambda([this](){ return GetDefault<UDisplayClusterLaunchEditorProjectSettings>()->bCloseEditorOnLaunch; })
+			),
+			NAME_None,
+			EUserInterfaceActionType::Check
+		);
+
+		MenuBuilder.AddMenuEntry(
+			LOCTEXT("AdvancedSettingsLabel", "Advanced Settings..."),
+			LOCTEXT("AdvancedSettingsTooltip", "Open the nDisplay Launch Project Settings"),
+			FSlateIcon(),
+			FUIAction(
+				FExecuteAction::CreateStatic(&FDisplayClusterLaunchEditorModule::OpenProjectSettings)
+			),
+			NAME_None
+		);
+	}
+	MenuBuilder.EndSection();
+}
+
+bool FDisplayClusterLaunchEditorModule::GetConnectToMultiUser() const
+{
+	return GetDefault<UDisplayClusterLaunchEditorProjectSettings>()->bConnectToMultiUser;
+}
+
+#undef LOCTEXT_NAMESPACE
+
+IMPLEMENT_MODULE(FDisplayClusterLaunchEditorModule, DisplayClusterLaunchEditor)
diff --git a/Engine/Plugins/Editor/DisplayClusterLaunch/Source/DisplayClusterLaunchEditor/Private/DisplayClusterLaunchEditorStyle.cpp b/Engine/Plugins/Editor/DisplayClusterLaunch/Source/DisplayClusterLaunchEditor/Private/DisplayClusterLaunchEditorStyle.cpp
new file mode 100644
--- /dev/null
+++ b/Engine/Plugins/Editor/DisplayClusterLaunch/Source/DisplayClusterLaunchEditor/Private/DisplayClusterLaunchEditorStyle.cpp
@@ -0,0 +1,112 @@
+// Copyright Epic Games, Inc. All Rights Reserved.
+
+#include "DisplayClusterLaunchEditorStyle.h"
+
+#include "EditorStyleSet.h"
+#include "Framework/Application/SlateApplication.h"
+#include "Interfaces/IPluginManager.h"
+#include "Styling/CoreStyle.h"
+#include "Styling/SlateStyleRegistry.h"
+#include "Styling/SlateTypes.h"
+#include "Styling/StyleColors.h"
+
+TSharedPtr<FSlateStyleSet> FDisplayClusterLaunchEditorStyle::StyleInstance = nullptr;
+
+void FDisplayClusterLaunchEditorStyle::Initialize()
+{
+	if (!StyleInstance.IsValid())
+	{
+		StyleInstance = Create();
+		FSlateStyleRegistry::RegisterSlateStyle(*StyleInstance);
+	}
+}
+
+void FDisplayClusterLaunchEditorStyle::Shutdown()
+{
+	FSlateStyleRegistry::UnRegisterSlateStyle(*StyleInstance);
+	ensure(StyleInstance.IsUnique());
+	StyleInstance.Reset();
+}
+
+void FDisplayClusterLaunchEditorStyle::ReloadTextures()
+{
+	FSlateApplication::Get().GetRenderer()->ReloadTextureResources();
+}
+
+const ISlateStyle& FDisplayClusterLaunchEditorStyle::Get()
+{
+	return *StyleInstance;
+}
+
+const FName& FDisplayClusterLaunchEditorStyle::GetStyleSetName() const
+{
+	static FName DisplayClusterLaunchStyleSetName(TEXT("DisplayClusterLaunchEditor"));
+	return DisplayClusterLaunchStyleSetName;
+}
+
+const FSlateBrush* FDisplayClusterLaunchEditorStyle::GetBrush(const FName PropertyName, const ANSICHAR* Specifier, const ISlateStyle* RequestingStyle) const
+{
+	return StyleInstance->GetBrush(PropertyName, Specifier);
+}
+
+#define IMAGE_BRUSH( RelativePath, ... ) FSlateImageBrush( Style->RootToContentDir( RelativePath, TEXT(".png") ), __VA_ARGS__ )
+#define BOX_BRUSH( RelativePath, ... ) FSlateBoxBrush( Style->RootToContentDir( RelativePath, TEXT(".png") ), __VA_ARGS__ )
+#define BORDER_BRUSH( RelativePath, ... ) FSlateBorderBrush( Style->RootToContentDir( RelativePath, TEXT(".png") ), __VA_ARGS__ )
+#define IMAGE_BRUSH_SVG( RelativePath, ... ) FSlateVectorImageBrush( Style->RootToContentDir(RelativePath, TEXT(".svg") ), __VA_ARGS__)
+#define IMAGE_PLUGIN_BRUSH_SVG( RelativePath, ... ) FSlateVectorImageBrush( FDisplayClusterLaunchEditorStyle::InContent(RelativePath, ".svg"), __VA_ARGS__)
+
+const FVector2D Icon64x64(64.f, 64.f);
+const FVector2D Icon40x40(40.0f, 40.0f);
+const FVector2D Icon20x20(20.0f, 20.0f);
+const FVector2D Icon16x16(16.0f, 16.0f);
+const FVector2D Icon12x12(12.0f, 12.0f);
+const FVector2D Icon8x8(8.f, 8.f);
+
+FString FDisplayClusterLaunchEditorStyle::InContent(const FString& RelativePath, const ANSICHAR* Extension)
+{
+	static FString ContentDir = IPluginManager::Get().FindPlugin(TEXT("ConcertSyncClient"))->GetContentDir();
+	return (ContentDir / RelativePath) + Extension;
+}
+
+TSharedRef< FSlateStyleSet > FDisplayClusterLaunchEditorStyle::Create()
+{
+	TSharedRef<FSlateStyleSet> Style = MakeShared<FSlateStyleSet>("DisplayClusterLaunchEditor");
+
+	TSharedPtr<IPlugin> Plugin = IPluginManager::Get().FindPlugin(TEXT("DisplayClusterLaunch"));
+	check(Plugin.IsValid());
+	if (Plugin.IsValid())
+	{
+		Style->SetContentRoot(FPaths::Combine(Plugin->GetBaseDir(), TEXT("Resources")));
+	}
+
+	Style->SetCoreContentRoot(FPaths::EngineContentDir() / TEXT("Slate"));
+
+	// Icons
+	// todo: update
+	Style->Set("ConsoleVariables.ToolbarButton", new IMAGE_BRUSH_SVG("Icons/ConsoleVariables", Icon40x40));
+	Style->Set("ConsoleVariables.ToolbarButton.Small", new IMAGE_BRUSH_SVG("Icons/ConsoleVariables", Icon20x20));
+	
+	Style->Set("ConsoleVariables.Favorite.Outline", new IMAGE_BRUSH_SVG("Icons/FavoriteOutline", Icon40x40));
+	Style->Set("ConsoleVariables.Favorite.Outline.Small", new IMAGE_BRUSH_SVG("Icons/FavoriteOutline", Icon16x16));
+	
+	Style->Set("ConsoleVariables.GlobalSearch", new IMAGE_BRUSH_SVG("Icons/SearchGlobal", Icon40x40));
+	Style->Set("ConsoleVariables.GlobalSearch.Small", new IMAGE_BRUSH_SVG("Icons/SearchGlobal", Icon16x16));
+
+	// Brush
+	Style->Set("ConsoleVariablesEditor.GroupBorder", new BOX_BRUSH("Common/DarkGroupBorder", FMargin(4.0f / 16.0f)));
+	Style->Set("ConsoleVariablesEditor.BrightBorder", new FSlateColorBrush(FStyleColors::Header));
+
+	// Border colors for Results view
+	Style->Set("ConsoleVariablesEditor.HeaderRowBorder", new FSlateColorBrush(FStyleColors::Black));
+	Style->Set("ConsoleVariablesEditor.CommandGroupBorder", new BOX_BRUSH("Common/DarkGroupBorder", FMargin(4.0f / 16.0f)));
+	Style->Set("ConsoleVariablesEditor.DefaultBorder", new FSlateColorBrush(FStyleColors::Transparent));
+
+	// Multi-user Tab/Menu icons
+	Style->Set("Concert.MultiUser", new IMAGE_PLUGIN_BRUSH_SVG("Icons/icon_MultiUser", Icon16x16));
+	
+	return Style;
+}
+
+#undef IMAGE_BRUSH
+#undef BOX_BRUSH
+#undef BORDER_BRUSH
diff --git a/Engine/Plugins/Editor/DisplayClusterLaunch/Source/DisplayClusterLaunchEditor/Public/DisplayClusterLaunchEditorLog.h b/Engine/Plugins/Editor/DisplayClusterLaunch/Source/DisplayClusterLaunchEditor/Public/DisplayClusterLaunchEditorLog.h
new file mode 100644
--- /dev/null
+++ b/Engine/Plugins/Editor/DisplayClusterLaunch/Source/DisplayClusterLaunchEditor/Public/DisplayClusterLaunchEditorLog.h
@@ -0,0 +1,7 @@
+// Copyright Epic Games, Inc. All Rights Reserved.
+
+#pragma once
+
+#include "CoreMinimal.h"
+
+DECLARE_LOG_CATEGORY_CLASS(LogDisplayClusterLaunchEditor, Log, All);
diff --git a/Engine/Plugins/Editor/DisplayClusterLaunch/Source/DisplayClusterLaunchEditor/Public/DisplayClusterLaunchEditorModule.h b/Engine/Plugins/Editor/DisplayClusterLaunch/Source/DisplayClusterLaunchEditor/Public/DisplayClusterLaunchEditorModule.h
new file mode 100644
--- /dev/null
+++ b/Engine/Plugins/Editor/DisplayClusterLaunch/Source/DisplayClusterLaunchEditor/Public/DisplayClusterLaunchEditorModule.h
@@ -0,0 +1,66 @@
+// Copyright Epic Games, Inc. All Rights Reserved.
+
+#pragma once
+
+#include "CoreMinimal.h"
+#include "Modules/ModuleInterface.h"
+#include "UObject/SoftObjectPath.h"
+
+class ADisplayClusterRootActor;
+class FMenuBuilder;
+class FUICommandList;
+class IAssetRegistry;
+class ISettingsSection;
+class SWidget;
+
+struct FAssetData;
+struct FSoftObjectPath;
+
+class FDisplayClusterLaunchEditorModule : public IModuleInterface
+{
+public:
+	
+	static FDisplayClusterLaunchEditorModule& Get();
+
+	//~ Begin IModuleInterface Interface
+	virtual void StartupModule() override;
+	virtual void ShutdownModule() override;
+	//~ End IModuleInterface Interface
+	
+	static void OpenProjectSettings();
+	
+	void LaunchDisplayClusterProcess();
+
+private:
+
+	void OnFEngineLoopInitComplete();
+
+	void RegisterToolbarItem();
+	void RemoveToolbarItem();
+	void RegisterProjectSettings() const;
+
+	TArray<TWeakObjectPtr<ADisplayClusterRootActor>> GetAllDisplayClusterConfigsInWorld();
+	bool DoesCurrentWorldHaveDisplayClusterConfig();
+
+	void SetSelectedDisplayClusterConfigActor(ADisplayClusterRootActor* SelectedActor);
+	void ToggleDisplayClusterConfigActorNodeSelected(FString InNodeName);
+	bool IsDisplayClusterConfigActorNodeSelected(FString InNodeName);
+	void SetSelectedConsoleVariablesAsset(const FAssetData InConsoleVariablesAsset);
+	
+	TSharedRef<SWidget> CreateToolbarMenuEntries();
+	void AddDisplayClusterLaunchConfigurations(
+		IAssetRegistry* AssetRegistry, FMenuBuilder& MenuBuilder, TArray<TWeakObjectPtr<ADisplayClusterRootActor>>& DisplayClusterConfigs);
+	void AddDisplayClusterLaunchNodes(class IAssetRegistry* AssetRegistry, FMenuBuilder& MenuBuilder);
+	void AddConsoleVariablesEditorAssetsToToolbarMenu(IAssetRegistry* AssetRegistry, FMenuBuilder& MenuBuilder);
+	void AddOptionsToToolbarMenu(FMenuBuilder& MenuBuilder);
+
+	bool GetConnectToMultiUser() const;
+
+	TSharedPtr<FUICommandList> Actions;
+
+	bool bAreConfigsFoundInWorld = false;
+
+	FSoftObjectPath SelectedDisplayClusterConfigActor;
+	TSet<FString> SelectedDisplayClusterConfigActorNodes;
+	FName SelectedConsoleVariablesAssetName = NAME_None;
+};
diff --git a/Engine/Plugins/Editor/DisplayClusterLaunch/Source/DisplayClusterLaunchEditor/Public/DisplayClusterLaunchEditorProjectSettings.h b/Engine/Plugins/Editor/DisplayClusterLaunch/Source/DisplayClusterLaunchEditor/Public/DisplayClusterLaunchEditorProjectSettings.h
new file mode 100644
--- /dev/null
+++ b/Engine/Plugins/Editor/DisplayClusterLaunch/Source/DisplayClusterLaunchEditor/Public/DisplayClusterLaunchEditorProjectSettings.h
@@ -0,0 +1,141 @@
+// Copyright Epic Games, Inc. All Rights Reserved.
+
+#pragma once
+
+#include "CoreMinimal.h"
+
+#include "DisplayClusterLaunchEditorProjectSettings.generated.h"
+
+UENUM()
+enum EDisplayClusterLaunchLogVerbosity
+{
+	/** Always prints a fatal error to console (and log file) and crashes (even if logging is disabled) */
+	Fatal,
+
+	/** 
+	 * Prints an error to console (and log file). 
+	 * Commandlets and the editor collect and report errors. Error messages result in commandlet failure.
+	 */
+	Error,
+
+	/** 
+	 * Prints a warning to console (and log file).
+	 * Commandlets and the editor collect and report warnings. Warnings can be treated as an error.
+	 */
+	Warning,
+
+	/** Prints a message to console (and log file) */
+	Display,
+
+	/** Prints a message to a log file (does not print to console) */
+	Log,
+
+	/** 
+	 * Prints a verbose message to a log file (if Verbose logging is enabled for the given category, 
+	 * usually used for detailed logging) 
+	 */
+	Verbose,
+
+	/** 
+	 * Prints a verbose message to a log file (if VeryVerbose logging is enabled, 
+	 * usually used for detailed logging that would otherwise spam output) 
+	 */
+	VeryVerbose
+};
+
+USTRUCT()
+struct FDisplayClusterLaunchLoggingConstruct
+{
+	GENERATED_BODY()
+
+	UPROPERTY(config, EditAnywhere, Category="nDisplay Launch Settings")
+	FName Category;
+
+	UPROPERTY(config, EditAnywhere, Category="nDisplay Launch Settings")
+	TEnumAsByte<EDisplayClusterLaunchLogVerbosity> VerbosityLevel = EDisplayClusterLaunchLogVerbosity::VeryVerbose;
+};
+
+UCLASS(config = DisplayClusterLaunch, defaultconfig)
+class DISPLAYCLUSTERLAUNCHEDITOR_API UDisplayClusterLaunchEditorProjectSettings : public UObject
+{
+	GENERATED_BODY()
+	
+public:
+
+	UDisplayClusterLaunchEditorProjectSettings(const FObjectInitializer& ObjectInitializer)
+	{
+		CommandLineArguments =
+		{
+			"messaging", "dc_cluster", "nosplash", "fixedseed", "NoVerifyGC", "noxrstereo", "xrtrackingonly", "RemoteControlIsHeadless",
+			"dx12", "dc_dev_mono", "unattended", "handleensurepercent=0", "fullscreen",
+			"ini:Engine:[/Script/Engine.Engine]:GameEngine=/Script/DisplayCluster.DisplayClusterGameEngine,[/Script/Engine.Engine]:GameViewportClientClassName=/Script/DisplayCluster.DisplayClusterViewportClient,[/Script/Engine.UserInterfaceSettings]:bAllowHighDPIInGameMode=True",
+			"ini:Game:[/Script/EngineSettings.GeneralProjectSettings]:bUseBorderlessWindow=True"
+		};
+
+		ConsoleCommands = { "DisableAllScreenMessages" };
+
+		DPCvars = { "p.Chaos.Solver.Deterministic=1", "r.Shadow.Virtual.Cache=0" };
+
+		Logging = { {"LogDisplayClusterRenderSync", EDisplayClusterLaunchLogVerbosity::VeryVerbose } };
+	}
+
+	/**
+	 * If true, this command will attempt to connect to an existing
+	 * session or create a new Multi-User session with the
+	 * options specified in the Multi-User Editing project settings.
+	 */
+	UPROPERTY(Config, EditAnywhere, Category="nDisplay Launch Settings")
+	bool bConnectToMultiUser = false;
+
+	/**
+	 * Whether or not to enable Unreal Insights for this session
+	 */
+	UPROPERTY(Config, EditAnywhere, Category="nDisplay Launch Settings")
+	bool bEnableUnrealInsights = false;
+
+	/**
+	 * If true, the editor will be closed on session launch to optimize session performance.
+	 */
+	UPROPERTY(Config, EditAnywhere, Category="nDisplay Launch Settings")
+	bool bCloseEditorOnLaunch = false;
+
+	/**
+	 * The name of the log file to which to write logs for the launched node.
+	 * '.log' will be automatically appended to the file name.
+	 * If not specified, the nDisplay node's name will be used instead.
+	 */
+	UPROPERTY(Config, EditAnywhere, Category="nDisplay Launch Settings")
+	FString LogFileName;
+
+	/**
+	 * A Console Variables Asset to always apply to your launches 
+	 */
+	UPROPERTY(Config, EditAnywhere, Category="nDisplay Launch Settings", meta = (AllowedClasses="ConsoleVariablesAsset"))
+	FSoftObjectPath ConsoleVariablesPreset;
+
+	/**
+	 * A list of command line arguments to append to the Switchboard command (e.g. messaging, fullscreen, handleensurepercent=0)
+	 * Do not include the dash ("-") as this will be automatically added for you when calling the command.
+	 * Parameters for arguments are supported, such as "handleensurepercent=0".
+	 */
+	UPROPERTY(Config, EditAnywhere, Category="nDisplay Launch Settings")
+	TSet<FString> CommandLineArguments;
+
+	/**
+	 * A list of Unreal Engine console commands to execute after executing the Switchboard command
+	 */
+	UPROPERTY(Config, EditAnywhere, Category="nDisplay Launch Settings")
+	TSet<FString> ConsoleCommands;
+
+	/**
+	 * DPCvars
+	 */
+	UPROPERTY(Config, EditAnywhere, Category="nDisplay Launch Settings")
+	TSet<FString> DPCvars;
+
+	/**
+	 * Which logs to include and with which verbosity level
+	 */
+	UPROPERTY(Config, EditAnywhere, Category="nDisplay Launch Settings")
+	TArray<FDisplayClusterLaunchLoggingConstruct> Logging;
+};
diff --git a/Engine/Plugins/Editor/DisplayClusterLaunch/Source/DisplayClusterLaunchEditor/Public/DisplayClusterLaunchEditorStyle.h b/Engine/Plugins/Editor/DisplayClusterLaunch/Source/DisplayClusterLaunchEditor/Public/DisplayClusterLaunchEditorStyle.h
new file mode 100644
--- /dev/null
+++ b/Engine/Plugins/Editor/DisplayClusterLaunch/Source/DisplayClusterLaunchEditor/Public/DisplayClusterLaunchEditorStyle.h
@@ -0,0 +1,36 @@
+// Copyright Epic Games, Inc. All Rights Reserved.
+
+#pragma once
+
+#include "Styling/SlateStyle.h"
+
+class FDisplayClusterLaunchEditorStyle : FSlateStyleSet
+{
+public:
+
+	static void Initialize();
+
+	static void Shutdown();
+
+	static void ReloadTextures();
+
+	static const ISlateStyle& Get();
+
+	virtual const FName& GetStyleSetName() const override;
+
+	virtual const FSlateBrush* GetBrush(const FName PropertyName, const ANSICHAR* Specifier = nullptr, const ISlateStyle* RequestingStyle = nullptr) const override;
+
+	template< typename WidgetStyleType >
+	static const WidgetStyleType& GetWidgetStyle(FName PropertyName, const ANSICHAR* Specifier = nullptr)
+	{
+		return StyleInstance->GetWidgetStyle<WidgetStyleType>(PropertyName, Specifier);
+	}
+
+private:
+	
+	static FString InContent(const FString& RelativePath, const ANSICHAR* Extension);
+	
+	static TSharedRef<FSlateStyleSet> Create();
+
+	static TSharedPtr<FSlateStyleSet> StyleInstance;
+};
