nDisplay Launch:

-Update uplugin from using Whitelist to Allowlist
-Add options to auto-generate or specify Concert Server and Session names
-Use content path instead of absolute path for Map asset argument
-Register settings before generating toolbar button
-Add launch state detection enum to inform the toolbar button icon/behaviour based on whether nodes have been launched
-Make default log level "Log" in project settings

diff --git a/Engine/Plugins/Editor/DisplayClusterLaunch/DisplayClusterLaunch.uplugin b/Engine/Plugins/Editor/DisplayClusterLaunch/DisplayClusterLaunch.uplugin
--- a/Engine/Plugins/Editor/DisplayClusterLaunch/DisplayClusterLaunch.uplugin
+++ b/Engine/Plugins/Editor/DisplayClusterLaunch/DisplayClusterLaunch.uplugin
@@ -43,7 +43,7 @@
 			"Name": "DisplayClusterLaunchEditor",
 			"Type": "Editor",
 			"LoadingPhase": "Default",
-			"WhitelistPlatforms": [ "Win64", "Linux" ]
+			"AllowlistPlatforms": [ "Win64", "Linux" ]
 		}
 	]
 }
diff --git a/Engine/Plugins/Editor/DisplayClusterLaunch/Source/DisplayClusterLaunchEditor/Private/DisplayClusterLaunchEditorModule.cpp b/Engine/Plugins/Editor/DisplayClusterLaunch/Source/DisplayClusterLaunchEditor/Private/DisplayClusterLaunchEditorModule.cpp
--- a/Engine/Plugins/Editor/DisplayClusterLaunch/Source/DisplayClusterLaunchEditor/Private/DisplayClusterLaunchEditorModule.cpp
+++ b/Engine/Plugins/Editor/DisplayClusterLaunch/Source/DisplayClusterLaunchEditor/Private/DisplayClusterLaunchEditorModule.cpp
@@ -190,6 +190,29 @@ bool AddUdpMessagingArguments(FString& ConcatenatedArguments)
 	return false;
 }
 
+FString AppendRandomNumbersToString(const FString InString, uint8 NumberToAppend = 6)
+{
+	FString RandomizedString = "_";
+
+	for (uint8 RandomIteration = 0; RandomIteration < NumberToAppend; RandomIteration++)
+	{
+		RandomizedString += FString::FromInt(FMath::RandRange(0, 9));
+	}
+	return InString + RandomizedString;
+}
+
+FString GetConcertServerName()
+{
+	const UDisplayClusterLaunchEditorProjectSettings* Settings = GetDefault<UDisplayClusterLaunchEditorProjectSettings>();
+	return Settings->bAutoGenerateServerName ? AppendRandomNumbersToString("nDisplayLaunchServer") : Settings->ExplicitServerName;
+}
+
+FString GetConcertSessionName()
+{
+	const UDisplayClusterLaunchEditorProjectSettings* Settings = GetDefault<UDisplayClusterLaunchEditorProjectSettings>();
+	return Settings->bAutoGenerateSessionName ? AppendRandomNumbersToString("nDisplayLaunchSession") : Settings->ExplicitSessionName;
+}
+
 FString GetConcertArguments()
 {
 	const UConcertClientConfig* ConcertClientConfig = GetDefault<UConcertClientConfig>();
@@ -197,7 +220,7 @@ FString GetConcertArguments()
 
 	FString ReturnValue =
 		FString::Printf(TEXT("-CONCERTISHEADLESS -CONCERTRETRYAUTOCONNECTONERROR -CONCERTAUTOCONNECT -CONCERTSERVER=\"%s\" -CONCERTSESSION=\"%s\""),
-			*ConcertClientConfig->DefaultServerURL, *ConcertClientConfig->DefaultSessionName);
+			 *GetConcertServerName(), *GetConcertSessionName());
 
 	return ReturnValue;
 }
@@ -266,13 +289,13 @@ void FDisplayClusterLaunchEditorModule::LaunchDisplayClusterProcess()
 	const FString EditorBinary = FPlatformProcess::ExecutablePath();
 	
 	const FString Project = FPaths::SetExtension(FPaths::Combine(FPaths::ProjectDir(), FApp::GetProjectName()),".uproject");
-	const FString Map = GetCurrentWorld()->GetCurrentLevel()->GetPackage()->GetLoadedPath().GetLocalFullPath();
+	const FString Map = GetCurrentWorld()->GetCurrentLevel()->GetPackage()->GetFName().ToString();
 
 	// Create Multi-user params
 	FString ConcertArguments;
 	if (GetConnectToMultiUser())
 	{
-		GetConcertArguments();
+		ConcertArguments = GetConcertArguments();
 		LaunchConcertServerIfNotRunning();
 	}
 
@@ -330,8 +353,8 @@ void FDisplayClusterLaunchEditorModule::OnFEngineLoopInitComplete()
 	Actions->MapAction(FDisplayClusterLaunchUiCommands::Get().LaunchDisplayCluster,
 		FExecuteAction::CreateRaw(this, &FDisplayClusterLaunchEditorModule::LaunchDisplayClusterProcess));
 	
-	RegisterToolbarItem();
 	RegisterProjectSettings();
+	RegisterToolbarItem();
 }
 
 void FDisplayClusterLaunchEditorModule::RegisterToolbarItem()
diff --git a/Engine/Plugins/Editor/DisplayClusterLaunch/Source/DisplayClusterLaunchEditor/Public/DisplayClusterLaunchEditorModule.h b/Engine/Plugins/Editor/DisplayClusterLaunch/Source/DisplayClusterLaunchEditor/Public/DisplayClusterLaunchEditorModule.h
--- a/Engine/Plugins/Editor/DisplayClusterLaunch/Source/DisplayClusterLaunchEditor/Public/DisplayClusterLaunchEditorModule.h
+++ b/Engine/Plugins/Editor/DisplayClusterLaunch/Source/DisplayClusterLaunchEditor/Public/DisplayClusterLaunchEditorModule.h
@@ -19,6 +19,12 @@ struct FSoftObjectPath;
 class FDisplayClusterLaunchEditorModule : public IModuleInterface
 {
 public:
+
+	enum class ELaunchState
+	{
+		NotLaunched,
+		Launched
+	};
 	
 	static FDisplayClusterLaunchEditorModule& Get();
 
diff --git a/Engine/Plugins/Editor/DisplayClusterLaunch/Source/DisplayClusterLaunchEditor/Public/DisplayClusterLaunchEditorProjectSettings.h b/Engine/Plugins/Editor/DisplayClusterLaunch/Source/DisplayClusterLaunchEditor/Public/DisplayClusterLaunchEditorProjectSettings.h
--- a/Engine/Plugins/Editor/DisplayClusterLaunch/Source/DisplayClusterLaunchEditor/Public/DisplayClusterLaunchEditorProjectSettings.h
+++ b/Engine/Plugins/Editor/DisplayClusterLaunch/Source/DisplayClusterLaunchEditor/Public/DisplayClusterLaunchEditorProjectSettings.h
@@ -76,7 +76,7 @@ public:
 
 		DPCvars = { "p.Chaos.Solver.Deterministic=1", "r.Shadow.Virtual.Cache=0" };
 
-		Logging = { {"LogDisplayClusterRenderSync", EDisplayClusterLaunchLogVerbosity::VeryVerbose } };
+		Logging = { {"LogDisplayClusterRenderSync", EDisplayClusterLaunchLogVerbosity::Log } };
 	}
 
 	/**
@@ -87,6 +87,36 @@ public:
 	UPROPERTY(Config, EditAnywhere, Category="nDisplay Launch Settings")
 	bool bConnectToMultiUser = false;
 
+	/**
+	 * If true, a server name will be automatically generated for you when connecting to multi-user.
+	 * If false, the text in ExplicitServerName will be used instead.
+	 * If ExplicitServerName is empty, a name will be generated whether this setting is true or false.
+	 */
+	UPROPERTY(Config, EditAnywhere, Category="nDisplay Launch Settings")
+	bool bAutoGenerateServerName = true;
+
+	/**
+	 * A specific server name to use when connecting to multi-user if bAutoGenerateServerName is false.
+	 * If left empty, a name will be generated whether bAutoGenerateServerName is true or false.
+	 */
+	UPROPERTY(Config, EditAnywhere, Category="nDisplay Launch Settings")
+	FString ExplicitServerName;
+
+	/**
+	 * If true, a session name will be automatically generated for you when connecting to multi-user.
+	 * If false, the text in ExplicitSessionName will be used instead.
+	 * If ExplicitSessionName is empty, a name will be generated whether this setting is true or false.
+	 */
+	UPROPERTY(Config, EditAnywhere, Category="nDisplay Launch Settings")
+	bool bAutoGenerateSessionName = true;
+
+	/**
+	 * A specific session name to use when connecting to multi-user if bAutoGenerateSessionName is false.
+	 * If left empty, a name will be generated whether bAutoGenerateSessionName is true or false.
+	 */
+	UPROPERTY(Config, EditAnywhere, Category="nDisplay Launch Settings")
+	FString ExplicitSessionName;
+
 	/**
 	 * Whether or not to enable Unreal Insights for this session
 	 */
