nDisplay Launch

-Change AllowListPlatforms in uplugin to PlatformAllowList

diff --git a/Engine/Plugins/Editor/DisplayClusterLaunch/DisplayClusterLaunch.uplugin b/Engine/Plugins/Editor/DisplayClusterLaunch/DisplayClusterLaunch.uplugin
--- a/Engine/Plugins/Editor/DisplayClusterLaunch/DisplayClusterLaunch.uplugin
+++ b/Engine/Plugins/Editor/DisplayClusterLaunch/DisplayClusterLaunch.uplugin
@@ -43,7 +43,7 @@
 			"Name": "DisplayClusterLaunchEditor",
 			"Type": "Editor",
 			"LoadingPhase": "Default",
-			"AllowlistPlatforms": [ "Win64", "Linux" ]
+			"PlatformAllowlist": [ "Win64", "Linux" ]
 		}
 	]
 }
