nDisplay Launch

-Icons pass (Toolbar and menu)

diff --git a/Engine/Plugins/Editor/DisplayClusterLaunch/DisplayClusterLaunch.uplugin b/Engine/Plugins/Editor/DisplayClusterLaunch/DisplayClusterLaunch.uplugin
--- a/Engine/Plugins/Editor/DisplayClusterLaunch/DisplayClusterLaunch.uplugin
+++ b/Engine/Plugins/Editor/DisplayClusterLaunch/DisplayClusterLaunch.uplugin
@@ -29,6 +29,10 @@
 			"Name": "ConcertMain",
 			"Enabled": true 
 		},
+		{
+			"Name": "ConcertSyncClient",
+			"Enabled": true
+		},
 		{
 			"Name": "MultiUserClient",
 			"Enabled": true
diff --git a/Engine/Plugins/Editor/DisplayClusterLaunch/Source/DisplayClusterLaunchEditor/Private/DisplayClusterLaunchEditorModule.cpp b/Engine/Plugins/Editor/DisplayClusterLaunch/Source/DisplayClusterLaunchEditor/Private/DisplayClusterLaunchEditorModule.cpp
--- a/Engine/Plugins/Editor/DisplayClusterLaunch/Source/DisplayClusterLaunchEditor/Private/DisplayClusterLaunchEditorModule.cpp
+++ b/Engine/Plugins/Editor/DisplayClusterLaunch/Source/DisplayClusterLaunchEditor/Private/DisplayClusterLaunchEditorModule.cpp
@@ -365,7 +365,11 @@ void FDisplayClusterLaunchEditorModule::RegisterToolbarItem()
 
 	FToolMenuSection& Section = Menu->AddSection("DisplayClusterLaunch");
 
-	FToolMenuEntry DisplayClusterLaunchButton = FToolMenuEntry::InitToolBarButton(FDisplayClusterLaunchUiCommands::Get().LaunchDisplayCluster);
+	FToolMenuEntry DisplayClusterLaunchButton =
+		FToolMenuEntry::InitToolBarButton(FDisplayClusterLaunchUiCommands::Get().LaunchDisplayCluster,
+			TAttribute<FText>(),
+			TAttribute<FText>(),
+			FSlateIcon(FAppStyle::Get().GetStyleSetName(), "Icons.Toolbar.Play"));
 	DisplayClusterLaunchButton.SetCommandList(Actions);
 
 	const FToolMenuEntry DisplayClusterLaunchComboButton = FToolMenuEntry::InitComboButton(
@@ -487,7 +491,7 @@ TSharedRef<SWidget>  FDisplayClusterLaunchEditorModule::CreateToolbarMenuEntries
 		MenuBuilder.AddMenuEntry(
 			LOCTEXT("DisplayClusterLaunchLastNode", "Launch Last Node"),
 			LOCTEXT("DisplayClusterLaunchLastNodeTooltip", "Launch the last node."),
-			FSlateIcon(),
+			FSlateIcon(FAppStyle::Get().GetStyleSetName(), "Icons.Toolbar.Play"),
 			FUIAction(
 				FExecuteAction::CreateRaw(this, &FDisplayClusterLaunchEditorModule::LaunchDisplayClusterProcess),
 				FCanExecuteAction::CreateRaw(this, &FDisplayClusterLaunchEditorModule::DoesCurrentWorldHaveDisplayClusterConfig)
@@ -544,7 +548,7 @@ void FDisplayClusterLaunchEditorModule::AddDisplayClusterLaunchConfigurations(IA
 				MenuBuilder.AddMenuEntry(
 					NodeName,
 					DisplayClusterConfigTooltip,
-					FSlateIcon(),
+					FSlateIcon(FDisplayClusterLaunchEditorStyle::Get().GetStyleSetName(),"Icons.DisplayCluster"),
 					FUIAction(
 						FExecuteAction::CreateRaw(this, &FDisplayClusterLaunchEditorModule::SetSelectedDisplayClusterConfigActor, Node.Get()),
 						FCanExecuteAction::CreateRaw(this, &FDisplayClusterLaunchEditorModule::DoesCurrentWorldHaveDisplayClusterConfig),
@@ -618,7 +622,9 @@ void FDisplayClusterLaunchEditorModule::AddDisplayClusterLaunchNodes(IAssetRegis
 					}),
 					FUIAction(FExecuteAction(), FCanExecuteAction::CreateRaw(this, &FDisplayClusterLaunchEditorModule::DoesCurrentWorldHaveDisplayClusterConfig)),
 					NAME_None,
-					EUserInterfaceActionType::None
+					EUserInterfaceActionType::None,
+					false,
+					FSlateIcon(FDisplayClusterLaunchEditorStyle::Get().GetStyleSetName(),"Icons.DisplayClusterNode")
 				);
 			}
 			MenuBuilder.EndSection();
@@ -664,7 +670,9 @@ void FDisplayClusterLaunchEditorModule::AddConsoleVariablesEditorAssetsToToolbar
 				}),
 				FUIAction(FExecuteAction(), FCanExecuteAction::CreateRaw(this, &FDisplayClusterLaunchEditorModule::DoesCurrentWorldHaveDisplayClusterConfig)),
 				NAME_None,
-				EUserInterfaceActionType::None
+				EUserInterfaceActionType::None,
+				false,
+				FSlateIcon(FDisplayClusterLaunchEditorStyle::Get().GetStyleSetName(),"Icons.ConsoleVariablesEditor")
 			);
 		}
 		MenuBuilder.EndSection();
@@ -678,7 +686,7 @@ void FDisplayClusterLaunchEditorModule::AddOptionsToToolbarMenu(FMenuBuilder& Me
 		MenuBuilder.AddMenuEntry(
 			LOCTEXT("ConnectToMultiUserLabel", "Connect to Multi-User"),
 			LOCTEXT("ConnectToMultiUserTooltip", "Connect to Multi-User"),
-			FSlateIcon(),
+			FSlateIcon(FDisplayClusterLaunchEditorStyle::Get().GetStyleSetName(),"Icons.MultiUser"),
 			FUIAction(
 				FExecuteAction::CreateLambda([this]()
 				{
@@ -695,7 +703,7 @@ void FDisplayClusterLaunchEditorModule::AddOptionsToToolbarMenu(FMenuBuilder& Me
 		MenuBuilder.AddMenuEntry(
 			LOCTEXT("EnableUnrealInsightsLabel", "Enable Unreal Insights"),
 			LOCTEXT("EnableUnrealInsightsTooltip", "Enable Unreal Insights"),
-			FSlateIcon(),
+			FSlateIcon(FAppStyle::Get().GetStyleSetName(),"UnrealInsights.MenuIcon"),
 			FUIAction(
 				FExecuteAction::CreateLambda([this]()
 				{
@@ -712,7 +720,7 @@ void FDisplayClusterLaunchEditorModule::AddOptionsToToolbarMenu(FMenuBuilder& Me
 		MenuBuilder.AddMenuEntry(
 			LOCTEXT("CloseEditorOnLaunchLabel", "Close Editor on Launch"),
 			LOCTEXT("CloseEditorOnLaunchTooltip", "Close Editor on Launch"),
-			FSlateIcon(),
+			FSlateIcon(FAppStyle::Get().GetStyleSetName(),"Icons.X"),
 			FUIAction(
 				FExecuteAction::CreateLambda([this]()
 				{
@@ -729,7 +737,7 @@ void FDisplayClusterLaunchEditorModule::AddOptionsToToolbarMenu(FMenuBuilder& Me
 		MenuBuilder.AddMenuEntry(
 			LOCTEXT("AdvancedSettingsLabel", "Advanced Settings..."),
 			LOCTEXT("AdvancedSettingsTooltip", "Open the nDisplay Launch Project Settings"),
-			FSlateIcon(),
+			FSlateIcon(FAppStyle::Get().GetStyleSetName(),"Icons.Settings"),
 			FUIAction(
 				FExecuteAction::CreateStatic(&FDisplayClusterLaunchEditorModule::OpenProjectSettings)
 			),
diff --git a/Engine/Plugins/Editor/DisplayClusterLaunch/Source/DisplayClusterLaunchEditor/Private/DisplayClusterLaunchEditorStyle.cpp b/Engine/Plugins/Editor/DisplayClusterLaunch/Source/DisplayClusterLaunchEditor/Private/DisplayClusterLaunchEditorStyle.cpp
--- a/Engine/Plugins/Editor/DisplayClusterLaunch/Source/DisplayClusterLaunchEditor/Private/DisplayClusterLaunchEditorStyle.cpp
+++ b/Engine/Plugins/Editor/DisplayClusterLaunch/Source/DisplayClusterLaunchEditor/Private/DisplayClusterLaunchEditorStyle.cpp
@@ -44,6 +44,12 @@ const FName& FDisplayClusterLaunchEditorStyle::GetStyleSetName() const
 	return DisplayClusterLaunchStyleSetName;
 }
 
+FString GetConcertContent(const FString& RelativePath, const ANSICHAR* Extension)
+{
+	static FString ContentDir = IPluginManager::Get().FindPlugin(TEXT("ConcertSyncClient"))->GetContentDir();
+	return (ContentDir / RelativePath) + Extension;
+}
+
 const FSlateBrush* FDisplayClusterLaunchEditorStyle::GetBrush(const FName PropertyName, const ANSICHAR* Specifier, const ISlateStyle* RequestingStyle) const
 {
 	return StyleInstance->GetBrush(PropertyName, Specifier);
@@ -53,7 +59,7 @@ const FSlateBrush* FDisplayClusterLaunchEditorStyle::GetBrush(const FName Proper
 #define BOX_BRUSH( RelativePath, ... ) FSlateBoxBrush( Style->RootToContentDir( RelativePath, TEXT(".png") ), __VA_ARGS__ )
 #define BORDER_BRUSH( RelativePath, ... ) FSlateBorderBrush( Style->RootToContentDir( RelativePath, TEXT(".png") ), __VA_ARGS__ )
 #define IMAGE_BRUSH_SVG( RelativePath, ... ) FSlateVectorImageBrush( Style->RootToContentDir(RelativePath, TEXT(".svg") ), __VA_ARGS__)
-#define IMAGE_PLUGIN_BRUSH_SVG( RelativePath, ... ) FSlateVectorImageBrush( FDisplayClusterLaunchEditorStyle::InContent(RelativePath, ".svg"), __VA_ARGS__)
+#define IMAGE_PLUGIN_BRUSH_SVG( PluginName, RelativePath, ... ) FSlateVectorImageBrush( FDisplayClusterLaunchEditorStyle::GetExternalPluginContent(PluginName, RelativePath, ".svg"), __VA_ARGS__)
 
 const FVector2D Icon64x64(64.f, 64.f);
 const FVector2D Icon40x40(40.0f, 40.0f);
@@ -62,10 +68,11 @@ const FVector2D Icon16x16(16.0f, 16.0f);
 const FVector2D Icon12x12(12.0f, 12.0f);
 const FVector2D Icon8x8(8.f, 8.f);
 
-FString FDisplayClusterLaunchEditorStyle::InContent(const FString& RelativePath, const ANSICHAR* Extension)
+FString FDisplayClusterLaunchEditorStyle::GetExternalPluginContent(const FString& PluginName, const FString& RelativePath, const ANSICHAR* Extension)
 {
-	static FString ContentDir = IPluginManager::Get().FindPlugin(TEXT("ConcertSyncClient"))->GetContentDir();
-	return (ContentDir / RelativePath) + Extension;
+	FString ContentDir = IPluginManager::Get().FindPlugin(PluginName)->GetBaseDir() / RelativePath + Extension;
+	UE_LOG(LogTemp, Error, TEXT("%hs: Content = %s"), __FUNCTION__, *ContentDir);
+	return ContentDir;
 }
 
 TSharedRef< FSlateStyleSet > FDisplayClusterLaunchEditorStyle::Create()
@@ -85,24 +92,14 @@ TSharedRef< FSlateStyleSet > FDisplayClusterLaunchEditorStyle::Create()
 	// todo: update
 	Style->Set("ConsoleVariables.ToolbarButton", new IMAGE_BRUSH_SVG("Icons/ConsoleVariables", Icon40x40));
 	Style->Set("ConsoleVariables.ToolbarButton.Small", new IMAGE_BRUSH_SVG("Icons/ConsoleVariables", Icon20x20));
+
+	// External plugin icons
 	
-	Style->Set("ConsoleVariables.Favorite.Outline", new IMAGE_BRUSH_SVG("Icons/FavoriteOutline", Icon40x40));
-	Style->Set("ConsoleVariables.Favorite.Outline.Small", new IMAGE_BRUSH_SVG("Icons/FavoriteOutline", Icon16x16));
+	Style->Set("Icons.ConsoleVariablesEditor", new IMAGE_PLUGIN_BRUSH_SVG("ConsoleVariables","Resources/Icons/ConsoleVariables", Icon16x16));
+	Style->Set("Icons.MultiUser", new IMAGE_PLUGIN_BRUSH_SVG("ConcertSharedSlate","Content/Icons/icon_MultiUser", Icon16x16));
+	Style->Set("Icons.DisplayCluster", new IMAGE_PLUGIN_BRUSH_SVG("nDisplay","Content/Icons/RootActor/nDisplay", Icon16x16));
+	Style->Set("Icons.DisplayClusterNode", new IMAGE_PLUGIN_BRUSH_SVG("nDisplay","Content/Icons/Cluster/ClusterNode", Icon16x16));
 	
-	Style->Set("ConsoleVariables.GlobalSearch", new IMAGE_BRUSH_SVG("Icons/SearchGlobal", Icon40x40));
-	Style->Set("ConsoleVariables.GlobalSearch.Small", new IMAGE_BRUSH_SVG("Icons/SearchGlobal", Icon16x16));
-
-	// Brush
-	Style->Set("ConsoleVariablesEditor.GroupBorder", new BOX_BRUSH("Common/DarkGroupBorder", FMargin(4.0f / 16.0f)));
-	Style->Set("ConsoleVariablesEditor.BrightBorder", new FSlateColorBrush(FStyleColors::Header));
-
-	// Border colors for Results view
-	Style->Set("ConsoleVariablesEditor.HeaderRowBorder", new FSlateColorBrush(FStyleColors::Black));
-	Style->Set("ConsoleVariablesEditor.CommandGroupBorder", new BOX_BRUSH("Common/DarkGroupBorder", FMargin(4.0f / 16.0f)));
-	Style->Set("ConsoleVariablesEditor.DefaultBorder", new FSlateColorBrush(FStyleColors::Transparent));
-
-	// Multi-user Tab/Menu icons
-	Style->Set("Concert.MultiUser", new IMAGE_PLUGIN_BRUSH_SVG("Icons/icon_MultiUser", Icon16x16));
 	
 	return Style;
 }
diff --git a/Engine/Plugins/Editor/DisplayClusterLaunch/Source/DisplayClusterLaunchEditor/Public/DisplayClusterLaunchEditorStyle.h b/Engine/Plugins/Editor/DisplayClusterLaunch/Source/DisplayClusterLaunchEditor/Public/DisplayClusterLaunchEditorStyle.h
--- a/Engine/Plugins/Editor/DisplayClusterLaunch/Source/DisplayClusterLaunchEditor/Public/DisplayClusterLaunchEditorStyle.h
+++ b/Engine/Plugins/Editor/DisplayClusterLaunch/Source/DisplayClusterLaunchEditor/Public/DisplayClusterLaunchEditorStyle.h
@@ -28,7 +28,7 @@ public:
 
 private:
 	
-	static FString InContent(const FString& RelativePath, const ANSICHAR* Extension);
+	static FString GetExternalPluginContent(const FString& PluginName, const FString& RelativePath, const ANSICHAR* Extension);
 	
 	static TSharedRef<FSlateStyleSet> Create();
 
