nDisplay Launch

-Apply nDisplay Config overrides

diff --git a/Engine/Plugins/Editor/DisplayClusterLaunch/Source/DisplayClusterLaunchEditor/Private/DisplayClusterLaunchEditorModule.cpp b/Engine/Plugins/Editor/DisplayClusterLaunch/Source/DisplayClusterLaunchEditor/Private/DisplayClusterLaunchEditorModule.cpp
--- a/Engine/Plugins/Editor/DisplayClusterLaunch/Source/DisplayClusterLaunchEditor/Private/DisplayClusterLaunchEditorModule.cpp
+++ b/Engine/Plugins/Editor/DisplayClusterLaunch/Source/DisplayClusterLaunchEditor/Private/DisplayClusterLaunchEditorModule.cpp
@@ -15,11 +15,12 @@
 
 #include "Shared/UdpMessagingSettings.h"
 
+#include "Algo/RemoveIf.h"
 #include "AssetRegistry/IAssetRegistry.h"
 #include "Editor.h"
 #include "EngineUtils.h"
+#include "FileHelpers.h"
 #include "Engine/GameEngine.h"
-#include "Framework/Commands/Commands.h"
 #include "ISettingsModule.h"
 #include "LevelEditor.h"
 #include "Misc/ConfigCacheIni.h"
@@ -27,6 +28,11 @@
 
 #define LOCTEXT_NAMESPACE "FDisplayClusterLaunchEditorModule"
 
+void CloseAllMenus()
+{
+	FSlateApplication::Get().DismissAllMenus();
+}
+
 FString EnumToString(const TCHAR* EnumName, const int32 EnumValue)
 {
 	const UEnum* EnumPtr = FindObject<UEnum>(ANY_PACKAGE, EnumName, true);
@@ -53,26 +59,6 @@ static UWorld* GetCurrentWorld()
 	return CurrentWorld;
 }
 
-class FDisplayClusterLaunchUiCommands : public TCommands<FDisplayClusterLaunchUiCommands>
-{
-public:
-	FDisplayClusterLaunchUiCommands()
-		: TCommands<FDisplayClusterLaunchUiCommands>(
-			"DisplayClusterLaunch",
-			LOCTEXT("DisplayClusterLaunchCommands", "DisplayClusterLaunch UI Commands"),
-			NAME_None,
-			FDisplayClusterLaunchEditorStyle::Get().GetStyleSetName()
-		)
-	{}
-
-	virtual void RegisterCommands() override
-	{
-		UI_COMMAND(LaunchDisplayCluster, "Launch nDisplay", "Launch nDisplay", EUserInterfaceActionType::Button, FInputChord());
-	}
-
-	TSharedPtr<FUICommandInfo> LaunchDisplayCluster;
-};
-
 FDisplayClusterLaunchEditorModule& FDisplayClusterLaunchEditorModule::Get()
 {
 	return FModuleManager::GetModuleChecked<FDisplayClusterLaunchEditorModule>("DisplayClusterLaunchEditor");
@@ -213,60 +199,70 @@ FString GetConcertSessionName()
 	return Settings->bAutoGenerateSessionName ? AppendRandomNumbersToString("nDisplayLaunchSession") : Settings->ExplicitSessionName;
 }
 
-FString GetConcertArguments()
+FString GetConcertArguments(const FString& ServerName, const FString& SessionName)
 {
 	const UConcertClientConfig* ConcertClientConfig = GetDefault<UConcertClientConfig>();
 	ensureAlwaysMsgf (ConcertClientConfig, TEXT("%hs: Unable to launch nDisplay because there is no UConcertClientConfig object."));
 
 	FString ReturnValue =
 		FString::Printf(TEXT("-CONCERTISHEADLESS -CONCERTRETRYAUTOCONNECTONERROR -CONCERTAUTOCONNECT -CONCERTSERVER=\"%s\" -CONCERTSESSION=\"%s\""),
-			 *GetConcertServerName(), *GetConcertSessionName());
+			 *ServerName, *SessionName);
 
 	return ReturnValue;
 }
 
-void LaunchConcertServerIfNotRunning()
+void LaunchConcertServerIfNotRunning(const FString& ServerName, const FString& SessionName)
 {
 	IMultiUserClientModule& MultiUserClientModule = IMultiUserClientModule::Get();
 	if (!MultiUserClientModule.IsConcertServerRunning())
 	{
-		MultiUserClientModule.LaunchConcertServer();
+		FServerLaunchOverrides Overrides;
+		Overrides.ServerName = ServerName;
+		
+		MultiUserClientModule.LaunchConcertServer(Overrides);
 	}
 }
 
 void FDisplayClusterLaunchEditorModule::LaunchDisplayClusterProcess()
 {
-	if (!SelectedDisplayClusterConfigActor.IsValid())
-	{
-		UE_LOG(LogDisplayClusterLaunchEditor, Error, TEXT("%hs: Unable to launch nDisplay because there is no valid selected nDisplay configuration."), __FUNCTION__);
-		return;
-	}
-
+	TArray<TWeakObjectPtr<ADisplayClusterRootActor>> ConfigsInWorld = GetAllDisplayClusterConfigsInWorld();
 	if (!DoesCurrentWorldHaveDisplayClusterConfig())
 	{
-		UE_LOG(LogDisplayClusterLaunchEditor, Error, TEXT("%hs: Unable to launch nDisplay because there are no valid selected nDisplay configurations in the world."), __FUNCTION__);
+		UE_LOG(LogDisplayClusterLaunchEditor, Error, TEXT("%hs: Unable to launch nDisplay because there are no valid nDisplay configurations in the world."), __FUNCTION__);
 		return;
 	}
 
-	if (SelectedDisplayClusterConfigActorNodes.Num() == 0)
+	if (!SelectedDisplayClusterConfigActor.IsValid())
 	{
-		UE_LOG(LogDisplayClusterLaunchEditor, Error, TEXT("%hs: Unable to launch nDisplay because there are no selected nDisplay nodes."), __FUNCTION__);
-		return;
+		for (const TWeakObjectPtr<ADisplayClusterRootActor>& Config : ConfigsInWorld)
+		{
+			if (ADisplayClusterRootActor* ConfigPtr = Config.Get())
+			{
+				TArray<FString> NodeNames;
+				ConfigPtr->GetConfigData()->Cluster->Nodes.GenerateKeyArray(NodeNames);
+
+				if (NodeNames.Num() > 0)
+				{
+					SetSelectedDisplayClusterConfigActor(ConfigPtr);
+					break;
+				}
+			}
+		}
 	}
 
+	UDisplayClusterConfigurationData* ConfigDataToUse = nullptr;
 	FString ConfigActorPath;
+	
+	// If it's valid we need to check the selected nodes against the current config. If they don't exist, we need to get the first one.
 	if (const ADisplayClusterRootActor* ConfigActor = Cast<ADisplayClusterRootActor>(SelectedDisplayClusterConfigActor.ResolveObject()))
 	{
-		const FString FilePath = FPaths::Combine(FPaths::ProjectSavedDir(), "Temp.ndisplay");
-		UDisplayClusterConfigurationData* ConfigDataCopy = DuplicateObject(ConfigActor->GetConfigData(), GetTransientPackage());
+		// Duplicate existing config data so we can make non-destructive edits
+		ConfigDataToUse = DuplicateObject(ConfigActor->GetConfigData(), GetTransientPackage());
 
-		if (!ConfigDataCopy->Scene)
-		{
-			ConfigDataCopy->Scene = NewObject<UDisplayClusterConfigurationScene>(ConfigDataCopy, NAME_None,
-			RF_ArchetypeObject | RF_Public );
-		}
-		
-		if (!ensureAlways(IDisplayClusterConfiguration::Get().SaveConfig(ConfigDataCopy, FilePath)))
+		ApplyDisplayClusterConfigOverrides(ConfigDataToUse);
+
+		const FString FilePath = FPaths::Combine(FPaths::ProjectSavedDir(), "Temp.ndisplay");
+		if (!ensureAlways(IDisplayClusterConfiguration::Get().SaveConfig(ConfigDataToUse, FilePath)))
 		{
 			UE_LOG(LogDisplayClusterLaunchEditor, Error, TEXT("%hs: Unable to launch nDisplay because the selected nDisplay Configuration could not be saved to a .ndisplay file. See the log for more information."), __FUNCTION__);
 			return;
@@ -295,8 +291,11 @@ void FDisplayClusterLaunchEditorModule::LaunchDisplayClusterProcess()
 	FString ConcertArguments;
 	if (GetConnectToMultiUser())
 	{
-		ConcertArguments = GetConcertArguments();
-		LaunchConcertServerIfNotRunning();
+		const FString ServerName = GetConcertServerName();
+		const FString SessionName = GetConcertSessionName();
+		
+		ConcertArguments = GetConcertArguments(ServerName, SessionName);
+		LaunchConcertServerIfNotRunning(ServerName, SessionName);
 	}
 
 	for (const FString& Node : SelectedDisplayClusterConfigActorNodes)
@@ -305,10 +304,59 @@ void FDisplayClusterLaunchEditorModule::LaunchDisplayClusterProcess()
 		FString ConcatenatedConsoleCommands;
 		FString ConcatenatedDPCvars;
 		FString ConcatenatedLogCommands;
+
+		// Fullscreen/Windowed
+		if (UDisplayClusterConfigurationClusterNode** NodePtrPtr = ConfigDataToUse->Cluster->Nodes.Find(Node))
+		{
+			if (const UDisplayClusterConfigurationClusterNode* NodePtr = *NodePtrPtr)
+			{
+				if (NodePtr->bIsFullscreen)
+				{
+					ConcatenatedCommandLineArguments += "-fullscreen ";
+				}
+				else
+				{
+					ConcatenatedCommandLineArguments +=
+						FString::Printf(
+							TEXT("-windowed -forceres -WinX=%i -WinY=%i, -ResX=%i, -ResY=%i "),
+								NodePtr->WindowRect.X, NodePtr->WindowRect.Y,
+								NodePtr->WindowRect.W, NodePtr->WindowRect.H
+						);
+				}
+			}
+		}
+		
+		// Open a modal to prompt for save, if dirty. Yes = Save & Continue. No = Continue Without Saving. Cancel = Stop Opening Assets.
+		UPackage* PackageToSave = nullptr;
+        		
+		if (UWorld* World = GetCurrentWorld())
+		{
+			if (ULevel* Level = World->GetCurrentLevel())
+			{
+				PackageToSave = Level->GetPackage();
+			}
+		}
+
+		if (PackageToSave)
+		{
+			const FEditorFileUtils::EPromptReturnCode DialogueResponse =
+					FEditorFileUtils::PromptForCheckoutAndSave(
+						{PackageToSave},
+						true,
+						true,
+						LOCTEXT("SavePackagesTitle", "Save Packages"),
+						LOCTEXT("ConfirmOpenLevelFormat", "Do you want to save the current level?\n\nCancel to abort launch.\n")
+					);
+		
+			if (DialogueResponse == FEditorFileUtils::EPromptReturnCode::PR_Cancelled)
+			{
+				return;
+			}
+		}
 	
 		GetProjectSettingsArguments(ProjectSettings, ConcatenatedCommandLineArguments, ConcatenatedConsoleCommands, ConcatenatedDPCvars,
 					 ConcatenatedLogCommands);
-
+		
 		AddUdpMessagingArguments(ConcatenatedCommandLineArguments);
 
 		// Add nDisplay node information
@@ -336,23 +384,30 @@ void FDisplayClusterLaunchEditorModule::LaunchDisplayClusterProcess()
 
 		void* WritePipe = nullptr;
 
-		FProcHandle Process = FPlatformProcess::CreateProc(
-			*EditorBinary, *Params,
-			ProjectSettings->bCloseEditorOnLaunch,
-			false, false, NULL,
-			0, NULL, WritePipe
+		ActiveProcesses.Add(
+			FPlatformProcess::CreateProc(
+				*EditorBinary, *Params,
+				ProjectSettings->bCloseEditorOnLaunch,
+				false, false, NULL,
+				0, NULL, WritePipe
+			)
 		);
 	}
 }
 
-void FDisplayClusterLaunchEditorModule::OnFEngineLoopInitComplete()
+void FDisplayClusterLaunchEditorModule::TerminateActiveDisplayClusterProcesses()
 {
-	Actions = MakeShared<FUICommandList>();
-	FDisplayClusterLaunchUiCommands::Register();
+	for (FProcHandle& Process : ActiveProcesses)
+	{
+		FPlatformProcess::TerminateProc(Process);
+		FPlatformProcess::CloseProc(Process);
+	}
 
-	Actions->MapAction(FDisplayClusterLaunchUiCommands::Get().LaunchDisplayCluster,
-		FExecuteAction::CreateRaw(this, &FDisplayClusterLaunchEditorModule::LaunchDisplayClusterProcess));
-	
+	RemoveTerminatedNodeProcesses();
+}
+
+void FDisplayClusterLaunchEditorModule::OnFEngineLoopInitComplete()
+{
 	RegisterProjectSettings();
 	RegisterToolbarItem();
 }
@@ -365,12 +420,19 @@ void FDisplayClusterLaunchEditorModule::RegisterToolbarItem()
 
 	FToolMenuSection& Section = Menu->AddSection("DisplayClusterLaunch");
 
-	FToolMenuEntry DisplayClusterLaunchButton =
-		FToolMenuEntry::InitToolBarButton(FDisplayClusterLaunchUiCommands::Get().LaunchDisplayCluster,
-			TAttribute<FText>(),
+	const FToolMenuEntry DisplayClusterLaunchButton =
+		FToolMenuEntry::InitToolBarButton("DisplayClusterLaunchToolbarButton",
+			FUIAction(
+				FExecuteAction::CreateRaw(this, &FDisplayClusterLaunchEditorModule::OnClickToolbarButton)
+			),
 			TAttribute<FText>(),
-			FSlateIcon(FAppStyle::Get().GetStyleSetName(), "Icons.Toolbar.Play"));
-	DisplayClusterLaunchButton.SetCommandList(Actions);
+			TAttribute<FText>::Create(
+				TAttribute<FText>::FGetter::CreateRaw(this, &FDisplayClusterLaunchEditorModule::GetToolbarButtonTooltipText)
+			),
+			TAttribute<FSlateIcon>::Create(
+				TAttribute<FSlateIcon>::FGetter::CreateRaw(this, &FDisplayClusterLaunchEditorModule::GetToolbarButtonIcon)
+			)
+		);
 
 	const FToolMenuEntry DisplayClusterLaunchComboButton = FToolMenuEntry::InitComboButton(
 		"DisplayClusterLaunchMenu",
@@ -386,6 +448,64 @@ void FDisplayClusterLaunchEditorModule::RegisterToolbarItem()
 	Section.AddEntry(DisplayClusterLaunchComboButton);
 }
 
+FText FDisplayClusterLaunchEditorModule::GetToolbarButtonTooltipText()
+{
+	if (ActiveProcesses.Num() == 1)
+	{
+		return LOCTEXT("TerminateActiveProcess","Terminate active nDisplay process");
+	}
+	else if (ActiveProcesses.Num() > 1)
+	{
+		return FText::Format(LOCTEXT("TerminateActiveProcessesFormat","Terminate {0} active nDisplay processes"), ActiveProcesses.Num());
+	}
+
+	if (!SelectedDisplayClusterConfigActor.ResolveObject())
+	{
+		return LOCTEXT("GenericLaunchDisplayClusterProcessText_NoConfig", "Launch an nDisplay instance using the first Config Actor found in the current level and the first node found in that configuration.\n\nSet specific configurations and nodes using the overflow menu.");
+	}
+
+	if (SelectedDisplayClusterConfigActorNodes.Num() == 0)
+	{
+		return FText::Format(
+			LOCTEXT("GenericLaunchDisplayClusterProcessText_NoNodesFormat", "Launch an nDisplay instance using the Config Actor named '{0}' and the first node found in this configuration.\n\nSet specific configurations and nodes using the overflow menu."),
+				FText::FromString(SelectedDisplayClusterConfigActor.GetAssetName())
+		);
+	}
+
+	FString ConfigActorName = SelectedDisplayClusterConfigActor.ResolveObject()->GetName();
+
+	const FString SplitTerm = "_C";
+	if (ConfigActorName.Contains(SplitTerm))
+	{
+		ConfigActorName = ConfigActorName.Left(ConfigActorName.Find(SplitTerm));
+	}
+	
+	return FText::Format(
+		LOCTEXT("LaunchDisplayClusterProcessesFormat", "Launch the following nodes:\n\n{0}\n\nFrom this configuration:\n\n{1}"),
+			GetSelectedNodesListText(), FText::FromString(ConfigActorName)
+		);
+}
+
+FSlateIcon FDisplayClusterLaunchEditorModule::GetToolbarButtonIcon()
+{
+	RemoveTerminatedNodeProcesses();
+
+	return FSlateIcon(FAppStyle::Get().GetStyleSetName(),
+		ActiveProcesses.Num() > 0 ? "Icons.Toolbar.Stop" : "Icons.Toolbar.Play");
+}
+
+void FDisplayClusterLaunchEditorModule::OnClickToolbarButton()
+{
+	if (ActiveProcesses.Num() == 0)
+	{
+		LaunchDisplayClusterProcess();
+	}
+	else
+	{
+		TerminateActiveDisplayClusterProcesses();
+	}
+}
+
 void FDisplayClusterLaunchEditorModule::RemoveToolbarItem()
 {
 	UToolMenu* Menu = UToolMenus::Get()->ExtendMenu("LevelEditor.LevelEditorToolBar.User");
@@ -408,6 +528,21 @@ void FDisplayClusterLaunchEditorModule::RegisterProjectSettings() const
 	}
 }
 
+FText FDisplayClusterLaunchEditorModule::GetSelectedNodesListText() const
+{
+	if (SelectedDisplayClusterConfigActorNodes.Num() > 0)
+	{
+		FString JoinedNodes = FString::Join(SelectedDisplayClusterConfigActorNodes, TEXT("\n"));
+		int32 IndexOfBreakOrLengthIfNoBreakFound = JoinedNodes.Find(TEXT("\n"));
+		if (IndexOfBreakOrLengthIfNoBreakFound == INDEX_NONE) { IndexOfBreakOrLengthIfNoBreakFound = JoinedNodes.Len(); }
+		JoinedNodes.InsertAt(IndexOfBreakOrLengthIfNoBreakFound, " ({0})");
+	
+		return FText::Format(FText::FromString(JoinedNodes), LOCTEXT("PrimaryNode", "Primary"));
+	}
+
+	return FText::GetEmpty();
+}
+
 TArray<TWeakObjectPtr<ADisplayClusterRootActor>> FDisplayClusterLaunchEditorModule::GetAllDisplayClusterConfigsInWorld()
 {
 	TArray<TWeakObjectPtr<ADisplayClusterRootActor>> CachedDisplayClusterActors;
@@ -434,6 +569,52 @@ bool FDisplayClusterLaunchEditorModule::DoesCurrentWorldHaveDisplayClusterConfig
 	return bAreConfigsFoundInWorld;
 }
 
+void FDisplayClusterLaunchEditorModule::ApplyDisplayClusterConfigOverrides(UDisplayClusterConfigurationData* ConfigDataCopy)
+{
+	if (!ConfigDataCopy->Scene)
+	{
+		ConfigDataCopy->Scene = NewObject<UDisplayClusterConfigurationScene>(ConfigDataCopy, NAME_None,
+		RF_ArchetypeObject | RF_Public );
+	}
+
+	// A Primary Node should always be automatically selected, but this code preempts a crash. Normally we use the PN specified in the UI.
+	// If one is not specified in the UI, we check to see if the primary node specified
+	// in the original config is in our node array selection from the UI.
+	// If it isn't, in the loop below we'll use the first active node.
+	bool bIsConfigPrimaryNodeInActiveNodes = false;
+	const bool bIsPrimaryNodeUnset = SelectedDisplayClusterConfigActorPrimaryNode.IsEmpty();
+	if (bIsPrimaryNodeUnset)
+	{
+		bIsConfigPrimaryNodeInActiveNodes =
+			SelectedDisplayClusterConfigActorNodes.Contains(ConfigDataCopy->Cluster->PrimaryNode.Id);
+	}
+	else
+	{
+		ConfigDataCopy->Cluster->PrimaryNode.Id = SelectedDisplayClusterConfigActorPrimaryNode;
+	}
+
+	TMap<FString, UDisplayClusterConfigurationClusterNode*> ActiveNodes;
+	TMap<FString, UDisplayClusterConfigurationClusterNode*> NodesInConfig = ConfigDataCopy->Cluster->Nodes;
+	for (int32 NodeIndex = 0; NodeIndex < SelectedDisplayClusterConfigActorNodes.Num(); NodeIndex++)
+	{
+		FString NodeId = SelectedDisplayClusterConfigActorNodes[NodeIndex];
+		if (UDisplayClusterConfigurationClusterNode** Node = NodesInConfig.Find(NodeId))
+		{
+			ActiveNodes.Add(NodeId, *Node);
+
+			(*Node)->Host = "127.0.0.1";
+
+			// If we haven't specified a primary node and the config's primary node is not in our selection, use the first active node.
+			if (bIsPrimaryNodeUnset && !bIsConfigPrimaryNodeInActiveNodes && ActiveNodes.Num() == 1)
+			{
+				ConfigDataCopy->Cluster->PrimaryNode.Id = NodeId;
+			}
+		}
+	}
+
+	ConfigDataCopy->Cluster->Nodes = ActiveNodes;
+}
+
 void FDisplayClusterLaunchEditorModule::SetSelectedDisplayClusterConfigActor(ADisplayClusterRootActor* SelectedActor)
 {
 	if (SelectedActor)
@@ -445,6 +626,8 @@ void FDisplayClusterLaunchEditorModule::SetSelectedDisplayClusterConfigActor(ADi
 			SelectedDisplayClusterConfigActor = AsSoftObjectPath;
 
 			SelectedDisplayClusterConfigActorNodes.Empty();
+
+			SelectFirstNode(SelectedActor);
 		}
 	}
 }
@@ -459,6 +642,18 @@ void FDisplayClusterLaunchEditorModule::ToggleDisplayClusterConfigActorNodeSelec
 	{
 		SelectedDisplayClusterConfigActorNodes.Add(InNodeName);
 	}
+
+	// Clear SelectedDisplayClusterConfigActorPrimaryNode if no nodes are selected
+	if (SelectedDisplayClusterConfigActorNodes.Num() == 0)
+	{
+		SelectedDisplayClusterConfigActorPrimaryNode = "";
+	}
+
+	// If a single node is selected, SelectedDisplayClusterConfigActorPrimaryNode must be this node
+	if (SelectedDisplayClusterConfigActorNodes.Num() == 1)
+	{
+		SelectedDisplayClusterConfigActorPrimaryNode = SelectedDisplayClusterConfigActorNodes[0];
+	}
 }
 
 bool FDisplayClusterLaunchEditorModule::IsDisplayClusterConfigActorNodeSelected(FString InNodeName)
@@ -478,19 +673,49 @@ void FDisplayClusterLaunchEditorModule::SetSelectedConsoleVariablesAsset(const F
 	}
 }
 
+void FDisplayClusterLaunchEditorModule::SelectFirstNode(ADisplayClusterRootActor* InConfig)
+{
+	TArray<FString> NodeNames;
+	InConfig->GetConfigData()->Cluster->Nodes.GenerateKeyArray(NodeNames);
+
+	if (NodeNames.Num() == 0)
+	{
+		UE_LOG(LogDisplayClusterLaunchEditor, Error, TEXT("%hs: Unable to launch nDisplay because there are no nDisplay nodes in the selected nDisplay Config named '{0}'."), __FUNCTION__, *InConfig->GetActorNameOrLabel());
+		return;
+	}
+		
+	SelectedDisplayClusterConfigActorNodes.SetNum(
+		Algo::StableRemoveIf(
+			SelectedDisplayClusterConfigActorNodes,
+			[NodeNames](const FString& SelectedNode)
+			{
+				return !NodeNames.Contains(SelectedNode);
+			}
+		)
+	);
+
+	if (SelectedDisplayClusterConfigActorNodes.Num() == 0)
+	{
+		const FString& NodeName = NodeNames[0];
+		UE_LOG(LogDisplayClusterLaunchEditor, Warning, TEXT("%hs: Selected nDisplay nodes were not found on the selected DisplayClusterRootActor. We will select the first valid node."), __FUNCTION__);
+		SelectedDisplayClusterConfigActorNodes.Add(NodeName);
+		UE_LOG(LogDisplayClusterLaunchEditor, Log, TEXT("%hs: Adding first valid node named '{0}' to selected nodes."), __FUNCTION__, *NodeName);
+	}
+}
+
 TSharedRef<SWidget>  FDisplayClusterLaunchEditorModule::CreateToolbarMenuEntries()
 {
 	IAssetRegistry* AssetRegistry = IAssetRegistry::Get();
 	
-	FMenuBuilder MenuBuilder(true, Actions);
+	FMenuBuilder MenuBuilder(false, nullptr);
 
 	TArray<TWeakObjectPtr<ADisplayClusterRootActor>> DisplayClusterConfigs = GetAllDisplayClusterConfigsInWorld();
 
 	MenuBuilder.BeginSection("DisplayClusterLaunch", LOCTEXT("DisplayClusterLauncher", "Launch nDisplay"));
 	{
 		MenuBuilder.AddMenuEntry(
-			LOCTEXT("DisplayClusterLaunchLastNode", "Launch Last Node"),
-			LOCTEXT("DisplayClusterLaunchLastNodeTooltip", "Launch the last node."),
+			LOCTEXT("DisplayClusterLaunchLastNode", "Launch Last Node Configuration"),
+			LOCTEXT("DisplayClusterLaunchLastNodeTooltip", "Launch the last node configuration."),
 			FSlateIcon(FAppStyle::Get().GetStyleSetName(), "Icons.Toolbar.Play"),
 			FUIAction(
 				FExecuteAction::CreateRaw(this, &FDisplayClusterLaunchEditorModule::LaunchDisplayClusterProcess),
@@ -572,63 +797,140 @@ void FDisplayClusterLaunchEditorModule::AddDisplayClusterLaunchConfigurations(IA
 
 void FDisplayClusterLaunchEditorModule::AddDisplayClusterLaunchNodes(IAssetRegistry* AssetRegistry, FMenuBuilder& MenuBuilder)
 {
-	if (ADisplayClusterRootActor* SelectedActor = Cast<ADisplayClusterRootActor>(SelectedDisplayClusterConfigActor.ResolveObject()))
+	MenuBuilder.BeginSection("DisplayClusterLaunchNodes", LOCTEXT("DisplayClusterLaunchNodes", "Nodes"));
 	{
-		TArray<FString> NodeNames;
-		SelectedActor->GetConfigData()->Cluster->Nodes.GenerateKeyArray(NodeNames);
-
-		if (NodeNames.Num() > 0)
-		{
-			MenuBuilder.BeginSection("DisplayClusterLaunchNodes", LOCTEXT("DisplayClusterLaunchNodes", "Nodes"));
-			{
-				const TAttribute<FText> SelectedDisplayClusterNodeNames =
-					TAttribute<FText>::Create(
-						[this]()
+		// Submenu for node selection. Using a WrapperSubMenu to avoid the menu automatically closing when selecting nodes
+		// AddWrapperSubMenu does not allow for TAttribute<FText> Labels, it just copies the FText input so we need this entry to display live data
+		MenuBuilder.AddMenuEntry(
+			TAttribute<FText>::CreateLambda(
+				[this] ()
+				{
+					const int32 NodeCount = SelectedDisplayClusterConfigActorNodes.Num();
+										
+					if (NodeCount > 0)
+					{
+						if (NodeCount == 1)
 						{
-							if (SelectedDisplayClusterConfigActorNodes.Num() > 0)
-							{
-								return FText::FromString(FString::Join(SelectedDisplayClusterConfigActorNodes, TEXT(", ")));
-							}
-							return LOCTEXT("NoDisplayClusterLaunchNodesSelected", "Please select nDisplay nodes to launch.");
+							return FText::Format(
+								LOCTEXT("SelectedSingleNodeFormat", "'{0}' Selected"),
+								FText::FromString(SelectedDisplayClusterConfigActorNodes[0]));
 						}
-					);
-				
-				MenuBuilder.AddSubMenu(
-					SelectedDisplayClusterNodeNames,
-					LOCTEXT("SelectDisplayClusterNodes","Select nDisplay Nodes"),
-					FNewMenuDelegate::CreateLambda([this, NodeNames](FMenuBuilder& NewMenuBuilder)
+						else
+						{
+							return FText::Format(
+								LOCTEXT("SelectedMultipleNodesFormat", "Selected {0} Nodes"), FText::AsNumber(NodeCount));
+						}
+					}
+					else
+					{
+						return LOCTEXT("NoDisplayClusterLaunchNodesSelected", "Please select nDisplay nodes to launch.");
+					}
+				}
+			),
+			TAttribute<FText>::CreateRaw(
+					this, &FDisplayClusterLaunchEditorModule::GetSelectedNodesListText
+			),
+			FSlateIcon(FDisplayClusterLaunchEditorStyle::Get().GetStyleSetName(),"Icons.DisplayClusterNode"),
+			FUIAction(
+				FExecuteAction(),
+				FCanExecuteAction::CreateLambda([]() { return false; })
+			),
+			NAME_None,
+			EUserInterfaceActionType::None
+		);
+		
+		MenuBuilder.AddWrapperSubMenu(
+			LOCTEXT("SelectDisplayClusterNodes","Select nDisplay Nodes"),
+			LOCTEXT("SelectDisplayClusterNodesTooltip","Select nDisplay Nodes.\nThe first node selected will be designated as the primary node when launched unless otherwise specified."),
+			FOnGetContent::CreateLambda([this]()
+			{
+				FMenuBuilder NewMenuBuilder(false, nullptr);
+
+				NewMenuBuilder.AddSubMenu(
+					TAttribute<FText>::CreateLambda([this]()
+					{
+						return
+							FText::Format(
+								LOCTEXT("SelectPrimaryNodeFormat","Select Primary Node ({0})"),
+								SelectedDisplayClusterConfigActorPrimaryNode.IsEmpty() ? LOCTEXT("None", "None") :
+								FText::FromString(SelectedDisplayClusterConfigActorPrimaryNode));
+					}),
+					LOCTEXT("SelectPrimaryNode", "Select the Primary Node"),
+					FNewMenuDelegate::CreateLambda([this](FMenuBuilder& InMenuBuilder)
 					{
-						for (const FString& NodeName : NodeNames)
+						const FText NodeTooltip =
+							LOCTEXT("MakePrimaryNodeTooltip", "Make this node the new Primary Node. Does not affect the original configuration.");
+						
+						for (const FString& SelectedNode : SelectedDisplayClusterConfigActorNodes)
 						{
-							const FText DisplayClusterNodeName = FText::FromString(NodeName);
-							const FText DisplayClusterNodeTooltip =
-								FText::Format(LOCTEXT("SelectDisplayClusterNodeFormat","Select node '{0}'"), DisplayClusterNodeName);
-							
-							NewMenuBuilder.AddMenuEntry(
-								DisplayClusterNodeName,
-								DisplayClusterNodeTooltip,
+							InMenuBuilder.AddMenuEntry(
+								FText::FromString(SelectedNode),
+								NodeTooltip,
 								FSlateIcon(),
 								FUIAction(
-									FExecuteAction::CreateRaw(
-										this, &FDisplayClusterLaunchEditorModule::ToggleDisplayClusterConfigActorNodeSelected, NodeName),
+									FExecuteAction::CreateLambda([this, SelectedNode]()
+									{
+										SelectedDisplayClusterConfigActorPrimaryNode = SelectedNode;
+									}),
 									FCanExecuteAction(),
-									FIsActionChecked::CreateRaw(
-										this, &FDisplayClusterLaunchEditorModule::IsDisplayClusterConfigActorNodeSelected, NodeName)
+									FIsActionChecked::CreateLambda([this, SelectedNode]()
+									{
+										return SelectedDisplayClusterConfigActorPrimaryNode == SelectedNode;
+									})
 								),
 								NAME_None,
-								EUserInterfaceActionType::Check
+								EUserInterfaceActionType::RadioButton
 							);
 						}
 					}),
-					FUIAction(FExecuteAction(), FCanExecuteAction::CreateRaw(this, &FDisplayClusterLaunchEditorModule::DoesCurrentWorldHaveDisplayClusterConfig)),
+					FUIAction(
+						FExecuteAction(),
+						FCanExecuteAction::CreateLambda([this]()
+						{
+							return SelectedDisplayClusterConfigActorNodes.Num() > 0;
+						})
+					),
 					NAME_None,
-					EUserInterfaceActionType::None,
-					false,
-					FSlateIcon(FDisplayClusterLaunchEditorStyle::Get().GetStyleSetName(),"Icons.DisplayClusterNode")
+					EUserInterfaceActionType::None
 				);
-			}
-			MenuBuilder.EndSection();
-		}
+
+				NewMenuBuilder.AddSeparator();
+
+				if (ADisplayClusterRootActor* SelectedActor = Cast<ADisplayClusterRootActor>(SelectedDisplayClusterConfigActor.ResolveObject()))
+				{
+					TArray<FString> NodeNames;
+				   SelectedActor->GetConfigData()->Cluster->Nodes.GenerateKeyArray(NodeNames);
+				
+				   for (const FString& NodeName : NodeNames)
+				   {
+					   const FText DisplayClusterNodeName = FText::FromString(NodeName);
+					
+					   const FText DisplayClusterNodeTooltip =
+						   FText::Format(LOCTEXT("SelectDisplayClusterNodeFormat","Select node '{0}'"), DisplayClusterNodeName);
+					
+					   NewMenuBuilder.AddMenuEntry(
+						   DisplayClusterNodeName,
+						   DisplayClusterNodeTooltip,
+						   FSlateIcon(),
+						   FUIAction(
+							   FExecuteAction::CreateRaw(
+								   this, &FDisplayClusterLaunchEditorModule::ToggleDisplayClusterConfigActorNodeSelected, NodeName),
+							   FCanExecuteAction(),
+							   FIsActionChecked::CreateRaw(
+								   this, &FDisplayClusterLaunchEditorModule::IsDisplayClusterConfigActorNodeSelected, NodeName)
+						   ),
+						   NAME_None,
+						   EUserInterfaceActionType::Check
+					   );
+				   }
+				}
+
+				return NewMenuBuilder.MakeWidget();
+			}),
+			FSlateIcon(FDisplayClusterLaunchEditorStyle::Get().GetStyleSetName(),"Icons.DisplayClusterNode")
+		);
+
+		MenuBuilder.EndSection();
 	}
 }
 
@@ -692,6 +994,7 @@ void FDisplayClusterLaunchEditorModule::AddOptionsToToolbarMenu(FMenuBuilder& Me
 				{
 					UDisplayClusterLaunchEditorProjectSettings* Settings = GetMutableDefault<UDisplayClusterLaunchEditorProjectSettings>();
 					Settings->bConnectToMultiUser = !GetConnectToMultiUser();
+					Settings->SaveConfig();
 				}),
 				FCanExecuteAction(),
 				FIsActionChecked::CreateRaw(this, &FDisplayClusterLaunchEditorModule::GetConnectToMultiUser)
@@ -709,6 +1012,7 @@ void FDisplayClusterLaunchEditorModule::AddOptionsToToolbarMenu(FMenuBuilder& Me
 				{
 					UDisplayClusterLaunchEditorProjectSettings* Settings = GetMutableDefault<UDisplayClusterLaunchEditorProjectSettings>();
 					Settings->bEnableUnrealInsights = !Settings->bEnableUnrealInsights;
+					Settings->SaveConfig();
 				}),
 				FCanExecuteAction(),
 				FIsActionChecked::CreateLambda([this](){ return GetDefault<UDisplayClusterLaunchEditorProjectSettings>()->bEnableUnrealInsights; })
@@ -726,6 +1030,7 @@ void FDisplayClusterLaunchEditorModule::AddOptionsToToolbarMenu(FMenuBuilder& Me
 				{
 					UDisplayClusterLaunchEditorProjectSettings* Settings = GetMutableDefault<UDisplayClusterLaunchEditorProjectSettings>();
 					Settings->bCloseEditorOnLaunch = !Settings->bCloseEditorOnLaunch;
+					Settings->SaveConfig();
 				}),
 				FCanExecuteAction(),
 				FIsActionChecked::CreateLambda([this](){ return GetDefault<UDisplayClusterLaunchEditorProjectSettings>()->bCloseEditorOnLaunch; })
@@ -752,6 +1057,19 @@ bool FDisplayClusterLaunchEditorModule::GetConnectToMultiUser() const
 	return GetDefault<UDisplayClusterLaunchEditorProjectSettings>()->bConnectToMultiUser;
 }
 
+void FDisplayClusterLaunchEditorModule::RemoveTerminatedNodeProcesses()
+{
+	ActiveProcesses.SetNum(
+		Algo::StableRemoveIf(
+			ActiveProcesses,
+			[](FProcHandle& Handle)
+			{
+				return !FPlatformProcess::IsProcRunning(Handle);
+			}
+		)
+	);
+}
+
 #undef LOCTEXT_NAMESPACE
 
 IMPLEMENT_MODULE(FDisplayClusterLaunchEditorModule, DisplayClusterLaunchEditor)
diff --git a/Engine/Plugins/Editor/DisplayClusterLaunch/Source/DisplayClusterLaunchEditor/Public/DisplayClusterLaunchEditorModule.h b/Engine/Plugins/Editor/DisplayClusterLaunch/Source/DisplayClusterLaunchEditor/Public/DisplayClusterLaunchEditorModule.h
--- a/Engine/Plugins/Editor/DisplayClusterLaunch/Source/DisplayClusterLaunchEditor/Public/DisplayClusterLaunchEditorModule.h
+++ b/Engine/Plugins/Editor/DisplayClusterLaunch/Source/DisplayClusterLaunchEditor/Public/DisplayClusterLaunchEditorModule.h
@@ -2,19 +2,26 @@
 
 #pragma once
 
-#include "CoreMinimal.h"
+#include "Containers/Array.h"
 #include "Modules/ModuleInterface.h"
+#include "Templates/SharedPointer.h"
 #include "UObject/SoftObjectPath.h"
+#include "UObject/WeakObjectPtrTemplates.h"
 
 class ADisplayClusterRootActor;
 class FMenuBuilder;
+class FName;
+class FString;
+class FText;
 class FUICommandList;
 class IAssetRegistry;
 class ISettingsSection;
 class SWidget;
 
 struct FAssetData;
+struct FProcHandle;
 struct FSoftObjectPath;
+struct FSlateIcon;
 
 class FDisplayClusterLaunchEditorModule : public IModuleInterface
 {
@@ -34,24 +41,35 @@ public:
 	//~ End IModuleInterface Interface
 	
 	static void OpenProjectSettings();
-	
+
 	void LaunchDisplayClusterProcess();
+	void TerminateActiveDisplayClusterProcesses();
 
 private:
 
 	void OnFEngineLoopInitComplete();
 
 	void RegisterToolbarItem();
+	FText GetToolbarButtonTooltipText();
+	FSlateIcon GetToolbarButtonIcon();
+	void OnClickToolbarButton();
 	void RemoveToolbarItem();
+	
 	void RegisterProjectSettings() const;
 
+	/** Returns a list of selected nodes as FText separated by new lines with the Primary Node marked. */
+	FText GetSelectedNodesListText() const;
+
 	TArray<TWeakObjectPtr<ADisplayClusterRootActor>> GetAllDisplayClusterConfigsInWorld();
 	bool DoesCurrentWorldHaveDisplayClusterConfig();
+	void ApplyDisplayClusterConfigOverrides(class UDisplayClusterConfigurationData* ConfigDataCopy);
 
 	void SetSelectedDisplayClusterConfigActor(ADisplayClusterRootActor* SelectedActor);
 	void ToggleDisplayClusterConfigActorNodeSelected(FString InNodeName);
 	bool IsDisplayClusterConfigActorNodeSelected(FString InNodeName);
 	void SetSelectedConsoleVariablesAsset(const FAssetData InConsoleVariablesAsset);
+
+	void SelectFirstNode(ADisplayClusterRootActor* InConfig);
 	
 	TSharedRef<SWidget> CreateToolbarMenuEntries();
 	void AddDisplayClusterLaunchConfigurations(
@@ -62,11 +80,15 @@ private:
 
 	bool GetConnectToMultiUser() const;
 
-	TSharedPtr<FUICommandList> Actions;
+	void RemoveTerminatedNodeProcesses();
 
 	bool bAreConfigsFoundInWorld = false;
 
 	FSoftObjectPath SelectedDisplayClusterConfigActor;
-	TSet<FString> SelectedDisplayClusterConfigActorNodes;
+	TArray<FString> SelectedDisplayClusterConfigActorNodes;
+	FString SelectedDisplayClusterConfigActorPrimaryNode;
+	
 	FName SelectedConsoleVariablesAssetName = NAME_None;
+
+	TArray<FProcHandle> ActiveProcesses; 
 };
diff --git a/Engine/Plugins/Editor/DisplayClusterLaunch/Source/DisplayClusterLaunchEditor/Public/DisplayClusterLaunchEditorProjectSettings.h b/Engine/Plugins/Editor/DisplayClusterLaunch/Source/DisplayClusterLaunchEditor/Public/DisplayClusterLaunchEditorProjectSettings.h
--- a/Engine/Plugins/Editor/DisplayClusterLaunch/Source/DisplayClusterLaunchEditor/Public/DisplayClusterLaunchEditorProjectSettings.h
+++ b/Engine/Plugins/Editor/DisplayClusterLaunch/Source/DisplayClusterLaunchEditor/Public/DisplayClusterLaunchEditorProjectSettings.h
@@ -67,7 +67,7 @@ public:
 		CommandLineArguments =
 		{
 			"messaging", "dc_cluster", "nosplash", "fixedseed", "NoVerifyGC", "noxrstereo", "xrtrackingonly", "RemoteControlIsHeadless",
-			"dx12", "dc_dev_mono", "unattended", "handleensurepercent=0", "fullscreen",
+			"dx12", "dc_dev_mono", "unattended", "handleensurepercent=0", 
 			"ini:Engine:[/Script/Engine.Engine]:GameEngine=/Script/DisplayCluster.DisplayClusterGameEngine,[/Script/Engine.Engine]:GameViewportClientClassName=/Script/DisplayCluster.DisplayClusterViewportClient,[/Script/Engine.UserInterfaceSettings]:bAllowHighDPIInGameMode=True",
 			"ini:Game:[/Script/EngineSettings.GeneralProjectSettings]:bUseBorderlessWindow=True"
 		};
