nDisplay Launch:

 -Warn the user if multi-user is already running and confirm that the user want to launch

diff --git a/Engine/Plugins/Editor/DisplayClusterLaunch/Source/DisplayClusterLaunchEditor/Private/DisplayClusterLaunchEditorModule.cpp b/Engine/Plugins/Editor/DisplayClusterLaunch/Source/DisplayClusterLaunchEditor/Private/DisplayClusterLaunchEditorModule.cpp
--- a/Engine/Plugins/Editor/DisplayClusterLaunch/Source/DisplayClusterLaunchEditor/Private/DisplayClusterLaunchEditorModule.cpp
+++ b/Engine/Plugins/Editor/DisplayClusterLaunch/Source/DisplayClusterLaunchEditor/Private/DisplayClusterLaunchEditorModule.cpp
@@ -21,9 +21,11 @@
 #include "EngineUtils.h"
 #include "FileHelpers.h"
 #include "Engine/GameEngine.h"
+#include "GenericPlatform/GenericPlatformMisc.h"
 #include "ISettingsModule.h"
 #include "LevelEditor.h"
 #include "Misc/ConfigCacheIni.h"
+#include "Misc/MessageDialog.h"
 #include "ToolMenus.h"
 
 #define LOCTEXT_NAMESPACE "FDisplayClusterLaunchEditorModule"
@@ -211,10 +213,9 @@ FString GetConcertArguments(const FString& ServerName, const FString& SessionNam
 	return ReturnValue;
 }
 
-void LaunchConcertServerIfNotRunning(const FString& ServerName, const FString& SessionName)
+void LaunchConcertServer(const FString& ServerName, const FString& SessionName)
 {
 	IMultiUserClientModule& MultiUserClientModule = IMultiUserClientModule::Get();
-	if (!MultiUserClientModule.IsConcertServerRunning())
 	{
 		FServerLaunchOverrides Overrides;
 		Overrides.ServerName = ServerName;
@@ -291,11 +292,28 @@ void FDisplayClusterLaunchEditorModule::LaunchDisplayClusterProcess()
 	FString ConcertArguments;
 	if (GetConnectToMultiUser())
 	{
-		const FString ServerName = GetConcertServerName();
-		const FString SessionName = GetConcertSessionName();
+		FString ServerName = GetConcertServerName();
+		FString SessionName = GetConcertSessionName();
+
+		IMultiUserClientModule& MultiUserClientModule = IMultiUserClientModule::Get();
+		if (MultiUserClientModule.IsConcertServerRunning())
+		{
+			EAppReturnType::Type Response =
+				FMessageDialog::Open(
+					EAppMsgType::YesNo,
+					LOCTEXT("MultiUserServerAlreadyRunningPrompt", "A Multi-user server is already running. Terminate the existing server and start a new instance? No will abort the nDisplay Launch command.")
+				);
+
+			if (Response == EAppReturnType::No)
+			{
+				return;
+			}
+
+			MultiUserClientModule.ShutdownConcertServer();
+		}
 		
 		ConcertArguments = GetConcertArguments(ServerName, SessionName);
-		LaunchConcertServerIfNotRunning(ServerName, SessionName);
+		LaunchConcertServer(ServerName, SessionName);
 	}
 
 	for (const FString& Node : SelectedDisplayClusterConfigActorNodes)
