nDisplay Launch:

-Add Unreal Insights support
-Minor UI polish

diff --git a/Engine/Plugins/Editor/DisplayClusterLaunch/Source/DisplayClusterLaunchEditor/Private/DisplayClusterLaunchEditorModule.cpp b/Engine/Plugins/Editor/DisplayClusterLaunch/Source/DisplayClusterLaunchEditor/Private/DisplayClusterLaunchEditorModule.cpp
--- a/Engine/Plugins/Editor/DisplayClusterLaunch/Source/DisplayClusterLaunchEditor/Private/DisplayClusterLaunchEditorModule.cpp
+++ b/Engine/Plugins/Editor/DisplayClusterLaunch/Source/DisplayClusterLaunchEditor/Private/DisplayClusterLaunchEditorModule.cpp
@@ -104,20 +104,20 @@ void GetProjectSettingsArguments(const UDisplayClusterLaunchEditorProjectSetting
 			{
 				continue;
 			}
-			ConcatenatedCommandLineArguments += FString::Printf(TEXT("-%s "), *CommandLineArgument);
+			ConcatenatedCommandLineArguments += FString::Printf(TEXT(" -%s "), *CommandLineArgument);
 		}
 		// Remove whitespace
 		ConcatenatedCommandLineArguments.TrimStartAndEndInline();
 	}
 
-	if (ProjectSettings->ConsoleCommands.Num() > 0)
+	if (ProjectSettings->AdditionalConsoleCommands.Num() > 0)
 	{
-		ConcatenatedConsoleCommands += FString::Join(ProjectSettings->ConsoleCommands, TEXT(","));
+		ConcatenatedConsoleCommands += FString::Join(ProjectSettings->AdditionalConsoleCommands, TEXT(","));
 	}
 
-	if (ProjectSettings->DPCvars.Num() > 0)
+	if (ProjectSettings->AdditionalConsoleVariables.Num() > 0)
 	{
-		ConcatenatedDPCvars += FString::Join(ProjectSettings->DPCvars, TEXT(","));
+		ConcatenatedDPCvars += FString::Join(ProjectSettings->AdditionalConsoleVariables, TEXT(","));
 	}
 
 	{
@@ -343,6 +343,27 @@ void FDisplayClusterLaunchEditorModule::LaunchDisplayClusterProcess()
 				}
 			}
 		}
+
+		// Unreal Insights support
+		if (ProjectSettings->bEnableUnrealInsights)
+		{
+			// Enable trace
+			ConcatenatedCommandLineArguments += " -trace";
+
+			if (ProjectSettings->bEnableStatNamedEvents)
+			{
+				ConcatenatedCommandLineArguments += " -statnamedevents";
+			}
+
+			// Override save directory if desired
+			if (ProjectSettings->bOverrideInsightsTraceFileSaveDirectory)
+			{
+				FString TraceFilePath =
+					ProjectSettings->ExplicitTraceFileSaveDirectory.Path / FDateTime::Now().ToString(TEXT("%Y%m%d_%H%M%S"));
+				ConcatenatedCommandLineArguments +=
+					FString::Printf(TEXT(" -tracefile=%s "), *TraceFilePath);
+			}
+		}
 		
 		// Open a modal to prompt for save, if dirty. Yes = Save & Continue. No = Continue Without Saving. Cancel = Stop Opening Assets.
 		UPackage* PackageToSave = nullptr;
diff --git a/Engine/Plugins/Editor/DisplayClusterLaunch/Source/DisplayClusterLaunchEditor/Public/DisplayClusterLaunchEditorProjectSettings.h b/Engine/Plugins/Editor/DisplayClusterLaunch/Source/DisplayClusterLaunchEditor/Public/DisplayClusterLaunchEditorProjectSettings.h
--- a/Engine/Plugins/Editor/DisplayClusterLaunch/Source/DisplayClusterLaunchEditor/Public/DisplayClusterLaunchEditorProjectSettings.h
+++ b/Engine/Plugins/Editor/DisplayClusterLaunch/Source/DisplayClusterLaunchEditor/Public/DisplayClusterLaunchEditorProjectSettings.h
@@ -72,19 +72,25 @@ public:
 			"ini:Game:[/Script/EngineSettings.GeneralProjectSettings]:bUseBorderlessWindow=True"
 		};
 
-		ConsoleCommands = { "DisableAllScreenMessages" };
+		AdditionalConsoleCommands = { "DisableAllScreenMessages" };
 
-		DPCvars = { "p.Chaos.Solver.Deterministic=1", "r.Shadow.Virtual.Cache=0" };
+		AdditionalConsoleVariables = { "p.Chaos.Solver.Deterministic=1", "r.Shadow.Virtual.Cache=0" };
 
 		Logging = { {"LogDisplayClusterRenderSync", EDisplayClusterLaunchLogVerbosity::Log } };
 	}
 
+	/**
+	 * If true, the editor will be closed on session launch to optimize session performance.
+	 */
+	UPROPERTY(Config, EditAnywhere, Category="nDisplay Launch Settings")
+	bool bCloseEditorOnLaunch = false;
+
 	/**
 	 * If true, this command will attempt to connect to an existing
 	 * session or create a new Multi-User session with the
 	 * options specified in the Multi-User Editing project settings.
 	 */
-	UPROPERTY(Config, EditAnywhere, Category="nDisplay Launch Settings")
+	UPROPERTY(Config, EditAnywhere, Category="nDisplay Launch Settings|Multi-User")
 	bool bConnectToMultiUser = false;
 
 	/**
@@ -92,14 +98,14 @@ public:
 	 * If false, the text in ExplicitServerName will be used instead.
 	 * If ExplicitServerName is empty, a name will be generated whether this setting is true or false.
 	 */
-	UPROPERTY(Config, EditAnywhere, Category="nDisplay Launch Settings")
+	UPROPERTY(Config, EditAnywhere, Category="nDisplay Launch Settings|Multi-User", meta = (EditCondition="bConnectToMultiUser"))
 	bool bAutoGenerateServerName = true;
 
 	/**
 	 * A specific server name to use when connecting to multi-user if bAutoGenerateServerName is false.
 	 * If left empty, a name will be generated whether bAutoGenerateServerName is true or false.
 	 */
-	UPROPERTY(Config, EditAnywhere, Category="nDisplay Launch Settings")
+	UPROPERTY(Config, EditAnywhere, Category="nDisplay Launch Settings|Multi-User", meta = (EditCondition="bConnectToMultiUser && !bAutoGenerateServerName"))
 	FString ExplicitServerName;
 
 	/**
@@ -107,65 +113,80 @@ public:
 	 * If false, the text in ExplicitSessionName will be used instead.
 	 * If ExplicitSessionName is empty, a name will be generated whether this setting is true or false.
 	 */
-	UPROPERTY(Config, EditAnywhere, Category="nDisplay Launch Settings")
+	UPROPERTY(Config, EditAnywhere, Category="nDisplay Launch Settings|Multi-User", meta = (EditCondition="bConnectToMultiUser"))
 	bool bAutoGenerateSessionName = true;
 
 	/**
 	 * A specific session name to use when connecting to multi-user if bAutoGenerateSessionName is false.
 	 * If left empty, a name will be generated whether bAutoGenerateSessionName is true or false.
 	 */
-	UPROPERTY(Config, EditAnywhere, Category="nDisplay Launch Settings")
+	UPROPERTY(Config, EditAnywhere, Category="nDisplay Launch Settings|Multi-User", meta = (EditCondition="bConnectToMultiUser && !bAutoGenerateSessionName"))
 	FString ExplicitSessionName;
 
 	/**
 	 * Whether or not to enable Unreal Insights for this session
 	 */
-	UPROPERTY(Config, EditAnywhere, Category="nDisplay Launch Settings")
+	UPROPERTY(Config, EditAnywhere, Category="nDisplay Launch Settings|Unreal Insights")
 	bool bEnableUnrealInsights = false;
 
 	/**
-	 * If true, the editor will be closed on session launch to optimize session performance.
+	 * Enable support for Stat Named Events in Unreal Insights
 	 */
-	UPROPERTY(Config, EditAnywhere, Category="nDisplay Launch Settings")
-	bool bCloseEditorOnLaunch = false;
+	UPROPERTY(Config, EditAnywhere, Category="nDisplay Launch Settings|Unreal Insights", meta = (EditCondition="bEnableUnrealInsights"))
+	bool bEnableStatNamedEvents = false;
 
 	/**
-	 * The name of the log file to which to write logs for the launched node.
-	 * '.log' will be automatically appended to the file name.
-	 * If not specified, the nDisplay node's name will be used instead.
+	 * If true, the Unreal Insights Trace Files will be saved to the path specified in ExplicitTraceFileSaveDirectory without needing to run Unreal Insights.
+	 * If false, Unreal Insights will connect to localhost (this computer) instead and you'll need to ensure Unreal Insights is launched.
+	 * To specify a socket to connect to, use CommandLineArguments.
 	 */
-	UPROPERTY(Config, EditAnywhere, Category="nDisplay Launch Settings")
-	FString LogFileName;
+	UPROPERTY(Config, EditAnywhere, Category="nDisplay Launch Settings|Unreal Insights", meta = (EditCondition="bEnableUnrealInsights"))
+	bool bOverrideInsightsTraceFileSaveDirectory = false;
+
+	/**
+	 * Where to store the Unreal Insights Trace Files if bOverrideInsightsTraceFileDirectory is true.
+	 */
+	UPROPERTY(Config, EditAnywhere, Category="nDisplay Launch Settings|Unreal Insights", meta = (EditCondition="bEnableUnrealInsights && bOverrideInsightsTraceFileSaveDirectory"))
+	FDirectoryPath ExplicitTraceFileSaveDirectory;
 
 	/**
 	 * A Console Variables Asset to always apply to your launches 
 	 */
-	UPROPERTY(Config, EditAnywhere, Category="nDisplay Launch Settings", meta = (AllowedClasses="ConsoleVariablesAsset"))
+	UPROPERTY(Config, EditAnywhere, Category="nDisplay Launch Settings|Console", meta = (AllowedClasses="ConsoleVariablesAsset"))
 	FSoftObjectPath ConsoleVariablesPreset;
 
+	/**
+	 * DPCvars
+	 * You can specify additional console variables here to be executed before those of the Console Variable Preset are executed
+	 */
+	UPROPERTY(Config, EditAnywhere, Category="nDisplay Launch Settings|Console", meta = (Keywords = "DPCvars"))
+	TSet<FString> AdditionalConsoleVariables;
+
+	/**
+	 * You can specify additional console commands here to be executed before those of the Console Variable Preset are executed
+	 */
+	UPROPERTY(Config, EditAnywhere, Category="nDisplay Launch Settings|Console")
+	TSet<FString> AdditionalConsoleCommands;
+
 	/**
 	 * A list of command line arguments to append to the Switchboard command (e.g. messaging, fullscreen, handleensurepercent=0)
 	 * Do not include the dash ("-") as this will be automatically added for you when calling the command.
 	 * Parameters for arguments are supported, such as "handleensurepercent=0".
 	 */
-	UPROPERTY(Config, EditAnywhere, Category="nDisplay Launch Settings")
+	UPROPERTY(Config, EditAnywhere, Category="nDisplay Launch Settings|Console")
 	TSet<FString> CommandLineArguments;
 
 	/**
-	 * A list of Unreal Engine console commands to execute after executing the Switchboard command
-	 */
-	UPROPERTY(Config, EditAnywhere, Category="nDisplay Launch Settings")
-	TSet<FString> ConsoleCommands;
-
-	/**
-	 * DPCvars
+	 * The name of the log file to which to write logs for the launched node.
+	 * '.log' will be automatically appended to the file name.
+	 * If not specified, the nDisplay node's name will be used instead.
 	 */
-	UPROPERTY(Config, EditAnywhere, Category="nDisplay Launch Settings")
-	TSet<FString> DPCvars;
+	UPROPERTY(Config, EditAnywhere, Category="nDisplay Launch Settings|Logging")
+	FString LogFileName;
 
 	/**
 	 * Which logs to include and with which verbosity level
 	 */
-	UPROPERTY(Config, EditAnywhere, Category="nDisplay Launch Settings")
+	UPROPERTY(Config, EditAnywhere, Category="nDisplay Launch Settings|Logging")
 	TArray<FDisplayClusterLaunchLoggingConstruct> Logging;
 };
