nDisplay Launch:

-MU flow changes

diff --git a/Engine/Plugins/Editor/DisplayClusterLaunch/Source/DisplayClusterLaunchEditor/DisplayClusterLaunchEditor.build.cs b/Engine/Plugins/Editor/DisplayClusterLaunch/Source/DisplayClusterLaunchEditor/DisplayClusterLaunchEditor.build.cs
--- a/Engine/Plugins/Editor/DisplayClusterLaunch/Source/DisplayClusterLaunchEditor/DisplayClusterLaunchEditor.build.cs
+++ b/Engine/Plugins/Editor/DisplayClusterLaunch/Source/DisplayClusterLaunchEditor/DisplayClusterLaunchEditor.build.cs
@@ -25,6 +25,7 @@ public class DisplayClusterLaunchEditor : ModuleRules
 				"Concert",
 				"ConcertTransport",
 				"CoreUObject",
+				"ConcertSyncClient",
 				"ConsoleVariablesEditor",
 				"ContentBrowser",
 				"DisplayCluster",
diff --git a/Engine/Plugins/Editor/DisplayClusterLaunch/Source/DisplayClusterLaunchEditor/Private/DisplayClusterLaunchEditorModule.cpp b/Engine/Plugins/Editor/DisplayClusterLaunch/Source/DisplayClusterLaunchEditor/Private/DisplayClusterLaunchEditorModule.cpp
--- a/Engine/Plugins/Editor/DisplayClusterLaunch/Source/DisplayClusterLaunchEditor/Private/DisplayClusterLaunchEditorModule.cpp
+++ b/Engine/Plugins/Editor/DisplayClusterLaunch/Source/DisplayClusterLaunchEditor/Private/DisplayClusterLaunchEditorModule.cpp
@@ -1,18 +1,20 @@
 // Copyright Epic Games, Inc. All Rights Reserved.
 
 #include "DisplayClusterLaunchEditorModule.h"
-
 #include "DisplayClusterLaunchEditorLog.h"
 #include "DisplayClusterLaunchEditorProjectSettings.h"
 #include "DisplayClusterLaunchEditorStyle.h"
 
-#include "ConcertSettings.h"
-#include "IMultiUserClientModule.h"
-
 #include "DisplayClusterRootActor.h"
 #include "IDisplayClusterConfiguration.h"
 #include "DisplayClusterConfigurationTypes.h"
 
+#include "ConcertSettings.h"
+#include "IConcertClient.h"
+#include "IConcertSyncClient.h"
+#include "IConcertSyncClientModule.h"
+#include "IMultiUserClientModule.h"
+
 #include "Shared/UdpMessagingSettings.h"
 
 #include "Algo/RemoveIf.h"
@@ -25,7 +27,7 @@
 #include "ISettingsModule.h"
 #include "LevelEditor.h"
 #include "Misc/ConfigCacheIni.h"
-#include "Misc/MessageDialog.h"
+ 
 #include "ToolMenus.h"
 
 #define LOCTEXT_NAMESPACE "FDisplayClusterLaunchEditorModule"
@@ -43,7 +45,6 @@ FString EnumToString(const TCHAR* EnumName, const int32 EnumValue)
 	{
 		return LOCTEXT("EnumNotFound", "Enum not found").ToString();
 	}
-
 	return EnumPtr->GetNameStringByIndex(EnumValue);
 }
 
@@ -69,30 +70,37 @@ FDisplayClusterLaunchEditorModule& FDisplayClusterLaunchEditorModule::Get()
 void FDisplayClusterLaunchEditorModule::StartupModule()
 {
 	FDisplayClusterLaunchEditorStyle::Initialize();
-
 	FCoreDelegates::OnFEngineLoopInitComplete.AddRaw(this, &FDisplayClusterLaunchEditorModule::OnFEngineLoopInitComplete);
 }
 
 void FDisplayClusterLaunchEditorModule::ShutdownModule()
 {
 	UToolMenus::UnregisterOwner(this);
-
 	FCoreDelegates::OnFEngineLoopInitComplete.RemoveAll(this);
-
 	FDisplayClusterLaunchEditorStyle::Shutdown();
-
+	
 	// Unregister project settings
 	ISettingsModule& SettingsModule = FModuleManager::LoadModuleChecked<ISettingsModule>("Settings");
 	{
 		SettingsModule.UnregisterSettings("Project", "Plugins", "nDisplay Launch");
 	}
+
+	// Remove Concert delegates
+	if (IConcertSyncClientModule* ConcertSyncClientModule = (IConcertSyncClientModule*)FModuleManager::Get().GetModule("ConcertSyncClient"))
+	{
+		if (const TSharedPtr<IConcertSyncClient> ConcertSyncClient = ConcertSyncClientModule->GetClient(TEXT("MultiUser")))
+		{
+			const IConcertClientRef ConcertClient = ConcertSyncClient->GetConcertClient();
+		
+			ConcertClient->OnKnownServersUpdated().RemoveAll(this);
+		}
+	}
 }
 
 void FDisplayClusterLaunchEditorModule::OpenProjectSettings()
 {
 	FModuleManager::LoadModuleChecked<ISettingsModule>("Settings")
 		.ShowViewer("Project", "Plugins", "nDisplay Launch");
-
 }
 
 void GetProjectSettingsArguments(const UDisplayClusterLaunchEditorProjectSettings* ProjectSettings, FString& ConcatenatedCommandLineArguments, FString& ConcatenatedConsoleCommands, FString& ConcatenatedDPCvars, FString& ConcatenatedLogCommands)
@@ -109,17 +117,14 @@ void GetProjectSettingsArguments(const UDisplayClusterLaunchEditorProjectSetting
 		// Remove whitespace
 		ConcatenatedCommandLineArguments.TrimStartAndEndInline();
 	}
-
 	if (ProjectSettings->AdditionalConsoleCommands.Num() > 0)
 	{
 		ConcatenatedConsoleCommands += FString::Join(ProjectSettings->AdditionalConsoleCommands, TEXT(","));
 	}
-
 	if (ProjectSettings->AdditionalConsoleVariables.Num() > 0)
 	{
 		ConcatenatedDPCvars += FString::Join(ProjectSettings->AdditionalConsoleVariables, TEXT(","));
 	}
-
 	{
 		for (const FDisplayClusterLaunchLoggingConstruct& LoggingConstruct : ProjectSettings->Logging)
 		{
@@ -146,20 +151,16 @@ bool AddUdpMessagingArguments(FString& ConcatenatedArguments)
 	{
 		TArray<FString> Settings;
 		FString Setting;
-
 		// Unicast endpoint setting
 		EngineConfig->GetString(TEXT("/Script/UdpMessaging.UdpMessagingSettings"), TEXT("UnicastEndpoint"), Setting);
-
 		// if the unicast endpoint port is bound, concatenate it
 		if (Setting != "0.0.0.0:0" && !Setting.IsEmpty())
 		{
 			ConcatenatedArguments += TEXT(" -UDPMESSAGING_TRANSPORT_UNICAST=") + Setting;
 		}
-
 		// Multicast endpoint setting
 		EngineConfig->GetString(TEXT("/Script/UdpMessaging.UdpMessagingSettings"), TEXT("MulticastEndpoint"), Setting);
 		ConcatenatedArguments += TEXT(" -UDPMESSAGING_TRANSPORT_MULTICAST=") + Setting;
-
 		// Static endpoints setting
 		Settings.Empty(1);
 		EngineConfig->GetArray(TEXT("/Script/UdpMessaging.UdpMessagingSettings"), TEXT("StaticEndpoints"), Settings);
@@ -181,58 +182,214 @@ bool AddUdpMessagingArguments(FString& ConcatenatedArguments)
 FString AppendRandomNumbersToString(const FString InString, uint8 NumberToAppend = 6)
 {
 	FString RandomizedString = "_";
-
 	for (uint8 RandomIteration = 0; RandomIteration < NumberToAppend; RandomIteration++)
 	{
 		RandomizedString += FString::FromInt(FMath::RandRange(0, 9));
 	}
 	return InString + RandomizedString;
 }
-
-FString GetConcertServerName()
-{
-	const UDisplayClusterLaunchEditorProjectSettings* Settings = GetDefault<UDisplayClusterLaunchEditorProjectSettings>();
-	return Settings->bAutoGenerateServerName ? AppendRandomNumbersToString("nDisplayLaunchServer") : Settings->ExplicitServerName;
-}
-
-FString GetConcertSessionName()
-{
-	const UDisplayClusterLaunchEditorProjectSettings* Settings = GetDefault<UDisplayClusterLaunchEditorProjectSettings>();
-	return Settings->bAutoGenerateSessionName ? AppendRandomNumbersToString("nDisplayLaunchSession") : Settings->ExplicitSessionName;
-}
-
+ 
 FString GetConcertArguments(const FString& ServerName, const FString& SessionName)
 {
 	const UConcertClientConfig* ConcertClientConfig = GetDefault<UConcertClientConfig>();
 	ensureAlwaysMsgf (ConcertClientConfig, TEXT("%hs: Unable to launch nDisplay because there is no UConcertClientConfig object."));
-
 	FString ReturnValue =
 		FString::Printf(TEXT("-CONCERTISHEADLESS -CONCERTRETRYAUTOCONNECTONERROR -CONCERTAUTOCONNECT -CONCERTSERVER=\"%s\" -CONCERTSESSION=\"%s\""),
 			 *ServerName, *SessionName);
-
 	return ReturnValue;
 }
 
-void LaunchConcertServer(const FString& ServerName, const FString& SessionName)
+void FDisplayClusterLaunchEditorModule::LaunchConcertServer()
 {
 	IMultiUserClientModule& MultiUserClientModule = IMultiUserClientModule::Get();
 	{
 		FServerLaunchOverrides Overrides;
-		Overrides.ServerName = ServerName;
+		Overrides.ServerName = GetConcertServerName();
+		ConcertServerRequestStatus = EConcertServerRequestStatus::LaunchRequested;
+		TOptional<FProcHandle> ServerHandle = MultiUserClientModule.LaunchConcertServer(Overrides);
+		if (ServerHandle.IsSet() && ServerHandle.GetValue().IsValid())
+		{
+			ServerTrackingData.MultiUserServerHandle = ServerHandle.GetValue();
+		}
+	}
+}
+
+void FDisplayClusterLaunchEditorModule::FindOrLaunchConcertServer()
+{
+	// Ensure we have the client, otherwise we can't do anything
+	if (const TSharedPtr<IConcertSyncClient> ConcertSyncClient = IConcertSyncClientModule::Get().GetClient(TEXT("MultiUser")))
+	{
+		const IConcertClientRef ConcertClient = ConcertSyncClient->GetConcertClient();
+		
+		ConcertClient->OnKnownServersUpdated().RemoveAll(this);
 		
-		MultiUserClientModule.LaunchConcertServer(Overrides);
+		// Shutdown existing server no matter what because we need to hook into OnServersAssumedReady
+		IMultiUserClientModule& MultiUserClientModule = IMultiUserClientModule::Get();
+		{
+			if (MultiUserClientModule.IsConcertServerRunning())
+			{
+				// Try to reuse last server
+				ConcertServerRequestStatus = EConcertServerRequestStatus::ReuseExisting;
+				OnServersAssumedReady();
+			}
+			else
+			{
+				// Continue when the server list is updated after creation
+				ConcertClient->OnKnownServersUpdated().AddRaw(
+					this, &FDisplayClusterLaunchEditorModule::OnServersAssumedReady
+				);
+
+				LaunchConcertServer();
+			}
+		}
+	}
+	else
+	{
+		UE_LOG(LogDisplayClusterLaunchEditor, Error, TEXT("%hs: The ConcertSyncClient could not be found. Please check the output log for errors and try again."), __FUNCTION__);
 	}
 }
 
-void FDisplayClusterLaunchEditorModule::LaunchDisplayClusterProcess()
+void FDisplayClusterLaunchEditorModule::OnServersAssumedReady()
+{
+	if (ConcertServerRequestStatus == EConcertServerRequestStatus::ShutdownRequested)
+	{
+		// If this method was called when trying to shut down the previous server
+		// then loop back so we can return after the new server is launched
+		ConcertServerRequestStatus = EConcertServerRequestStatus::None;
+		FindOrLaunchConcertServer();
+	}
+	else if (ConcertServerRequestStatus == EConcertServerRequestStatus::LaunchRequested ||
+		ConcertServerRequestStatus == EConcertServerRequestStatus::ReuseExisting)
+	{
+		FindAppropriateServer();
+	}
+	else
+	{
+		UE_LOG(LogDisplayClusterLaunchEditor, Warning, TEXT("%hs: OnServersAssumedReady was called when ConcertServerRequestStatus was None."), __FUNCTION__);
+	}
+}
+
+void FDisplayClusterLaunchEditorModule::FindAppropriateServer()
+{
+	ConcertServerRequestStatus = EConcertServerRequestStatus::None;
+	if (const TSharedPtr<IConcertSyncClient> ConcertSyncClient = IConcertSyncClientModule::Get().GetClient(TEXT("MultiUser")))
+	{
+		const IConcertClientRef ConcertClient = ConcertSyncClient->GetConcertClient();
+		if (ConcertClient->GetKnownServers().Num() == 0)
+		{
+			UE_LOG(LogDisplayClusterLaunchEditor, Warning, TEXT("%hs: No servers found. Please launch and connect to one manually."), __FUNCTION__);
+			return;
+		}
+
+		// Try to connect to an existing session even if we launched a new server
+		if (const TSharedPtr<IConcertClientSession> ConcertClientSession = ConcertClient->GetCurrentSession())
+		{
+			const FConcertSessionInfo& SessionInfo = ConcertClientSession->GetSessionInfo();
+
+			// Ensure the reported server is actually running then pull the latest data from it
+			bool bFoundMatch = false;
+			for (const FConcertServerInfo& ServerInfo : ConcertClient->GetKnownServers())
+			{
+				if (ServerInfo.InstanceInfo.InstanceId == SessionInfo.ServerInstanceId)
+				{
+					ServerTrackingData.MultiUserServerInfo = ServerInfo;
+					ServerTrackingData.GeneratedMultiUserServerName = ServerInfo.ServerName;
+					CachedConcertSessionName = ConcertClientSession->GetName();
+					bFoundMatch = true;
+					break;
+				}
+			}
+
+			if (bFoundMatch)
+			{
+				ConnectToSession();
+			}
+			else
+			{
+				UE_LOG(LogDisplayClusterLaunchEditor, Error, TEXT("%hs: ConcertClientSession reported a connected server but the server is not in the known servers list."), __FUNCTION__);
+			}
+		}
+		else // If no session, we need to try to find a server with a name matching the cached name
+		{
+			bool bFoundMatch = false;
+			for (const FConcertServerInfo& ServerInfo : ConcertClient->GetKnownServers())
+			{
+				if (ServerInfo.ServerName == GetConcertServerName())
+				{
+					ServerTrackingData.MultiUserServerInfo = ServerInfo;
+					bFoundMatch = true;
+					break;
+				}
+			}
+
+			if (bFoundMatch)
+			{
+				ConnectToSession();
+			}
+			else
+			{
+				UE_LOG(LogDisplayClusterLaunchEditor, Error, TEXT("%hs: Servers exist but a matching server was not found. Try connecting to a server and session manually."), __FUNCTION__);
+			}
+		}
+	}
+	else
+	{
+		UE_LOG(LogDisplayClusterLaunchEditor, Error, TEXT("%hs: The ConcertSyncClient could not be found. Please check the output log for errors and try again."), __FUNCTION__);
+	}
+}
+
+void FDisplayClusterLaunchEditorModule::ConnectToSession()
 {
+	// Session Management
+	// First check to see if we're in a session already
+	if (const TSharedPtr<IConcertSyncClient> ConcertSyncClient = IConcertSyncClientModule::Get().GetClient(TEXT("MultiUser")))
+	{
+		const IConcertClientRef ConcertClient = ConcertSyncClient->GetConcertClient();
+		if (const TSharedPtr<IConcertClientSession> ConcertClientSession = ConcertClient->GetCurrentSession())
+		{
+			// If we're already connected, go straight into launch
+			LaunchDisplayClusterProcess();
+			return;
+		}
+		
+		const UConcertClientConfig* CurrentConfig = ConcertClient->GetConfiguration();
+		UConcertClientConfig* AutoConnectConfig = DuplicateObject(CurrentConfig, GetTransientPackage(), CurrentConfig->GetFName());
+		AutoConnectConfig->bAutoConnect = true;
+		AutoConnectConfig->bRetryAutoConnectOnError = true;
+		AutoConnectConfig->DefaultServerURL = GetConcertServerName();
+		AutoConnectConfig->DefaultSessionName = GetConcertSessionName();
+		
+		ConcertClient->Configure(AutoConnectConfig);
+		check(ConcertClient->IsConfigured());
+		
+		// Initiate the auto connect to the named server and session.
+		if (ConcertClient->CanAutoConnect())
+		{
+			ConcertClient->StartAutoConnect();
+			LaunchDisplayClusterProcess();
+		}
+		else
+		{
+			ConcertClient->Configure(CurrentConfig);
+			UE_LOG(LogDisplayClusterLaunchEditor, Error, TEXT("Unable to start Multi-user auto connect routines."));
+		}
+	}
+}
+
+void FDisplayClusterLaunchEditorModule::TryLaunchDisplayClusterProcess()
+{
+	if (!ensureAlwaysMsgf(GetDefault<UDisplayClusterLaunchEditorProjectSettings>(), TEXT("%hs: Unable to launch nDisplay because there is no UDisplayClusterLaunchEditorProjectSettings object.")))
+	{
+		return;
+	}
+	
 	TArray<TWeakObjectPtr<ADisplayClusterRootActor>> ConfigsInWorld = GetAllDisplayClusterConfigsInWorld();
 	if (!DoesCurrentWorldHaveDisplayClusterConfig())
 	{
 		UE_LOG(LogDisplayClusterLaunchEditor, Error, TEXT("%hs: Unable to launch nDisplay because there are no valid nDisplay configurations in the world."), __FUNCTION__);
 		return;
 	}
-
+	
 	if (!SelectedDisplayClusterConfigActor.IsValid())
 	{
 		for (const TWeakObjectPtr<ADisplayClusterRootActor>& Config : ConfigsInWorld)
@@ -241,7 +398,6 @@ void FDisplayClusterLaunchEditorModule::LaunchDisplayClusterProcess()
 			{
 				TArray<FString> NodeNames;
 				ConfigPtr->GetConfigData()->Cluster->Nodes.GenerateKeyArray(NodeNames);
-
 				if (NodeNames.Num() > 0)
 				{
 					SetSelectedDisplayClusterConfigActor(ConfigPtr);
@@ -250,7 +406,22 @@ void FDisplayClusterLaunchEditorModule::LaunchDisplayClusterProcess()
 			}
 		}
 	}
+	
+	// Create Multi-user params async
+	if (GetConnectToMultiUser())
+	{
+		CachedConcertSessionName.Empty();
+		FindOrLaunchConcertServer();
+	}
+	else
+	{
+		LaunchDisplayClusterProcess();
+	}
+}
 
+void FDisplayClusterLaunchEditorModule::LaunchDisplayClusterProcess()
+{
+	const UDisplayClusterLaunchEditorProjectSettings* ProjectSettings = GetDefault<UDisplayClusterLaunchEditorProjectSettings>();
 	UDisplayClusterConfigurationData* ConfigDataToUse = nullptr;
 	FString ConfigActorPath;
 	
@@ -259,9 +430,7 @@ void FDisplayClusterLaunchEditorModule::LaunchDisplayClusterProcess()
 	{
 		// Duplicate existing config data so we can make non-destructive edits
 		ConfigDataToUse = DuplicateObject(ConfigActor->GetConfigData(), GetTransientPackage());
-
 		ApplyDisplayClusterConfigOverrides(ConfigDataToUse);
-
 		const FString FilePath = FPaths::Combine(FPaths::ProjectSavedDir(), "Temp.ndisplay");
 		if (!ensureAlways(IDisplayClusterConfiguration::Get().SaveConfig(ConfigDataToUse, FilePath)))
 		{
@@ -277,44 +446,48 @@ void FDisplayClusterLaunchEditorModule::LaunchDisplayClusterProcess()
 		return;
 	}
 	
-	const UDisplayClusterLaunchEditorProjectSettings* ProjectSettings = GetDefault<UDisplayClusterLaunchEditorProjectSettings>();
-	if (!ensureAlwaysMsgf (ProjectSettings, TEXT("%hs: Unable to launch nDisplay because there is no UDisplayClusterLaunchEditorProjectSettings object.")))
-	{
-		return;
-	}
+	UE_LOG(LogDisplayClusterLaunchEditor, Log, TEXT("%hs: Launching nDisplay processes..."), __FUNCTION__);
 	
-	const FString EditorBinary = FPlatformProcess::ExecutablePath();
-	
-	const FString Project = FPaths::SetExtension(FPaths::Combine(FPaths::ProjectDir(), FApp::GetProjectName()),".uproject");
-	const FString Map = GetCurrentWorld()->GetCurrentLevel()->GetPackage()->GetFName().ToString();
-
-	// Create Multi-user params
 	FString ConcertArguments;
+
 	if (GetConnectToMultiUser())
 	{
-		FString ServerName = GetConcertServerName();
-		FString SessionName = GetConcertSessionName();
-
-		IMultiUserClientModule& MultiUserClientModule = IMultiUserClientModule::Get();
-		if (MultiUserClientModule.IsConcertServerRunning())
+		ConcertArguments = GetConcertArguments(GetConcertServerName(), GetConcertSessionName());
+	}
+	else
+	{
+		// Open a modal to prompt for save, if dirty. Yes = Save & Continue. No = Continue Without Saving. Cancel = Stop Opening Assets.
+		UPackage* PackageToSave = nullptr;
+        		
+		if (UWorld* World = GetCurrentWorld())
 		{
-			EAppReturnType::Type Response =
-				FMessageDialog::Open(
-					EAppMsgType::YesNo,
-					LOCTEXT("MultiUserServerAlreadyRunningPrompt", "A Multi-user server is already running. Terminate the existing server and start a new instance? No will abort the nDisplay Launch command.")
-				);
-
-			if (Response == EAppReturnType::No)
+			if (ULevel* Level = World->GetCurrentLevel())
 			{
-				return;
+				PackageToSave = Level->GetPackage();
 			}
-
-			MultiUserClientModule.ShutdownConcertServer();
 		}
+		if (PackageToSave)
+		{
+			const FEditorFileUtils::EPromptReturnCode DialogueResponse =
+					FEditorFileUtils::PromptForCheckoutAndSave(
+						{PackageToSave},
+						true,
+						true,
+						LOCTEXT("SavePackagesTitle", "Save Packages"),
+						LOCTEXT("ConfirmOpenLevelFormat", "Do you want to save the current level?\n\nCancel to abort launch.\n")
+					);
 		
-		ConcertArguments = GetConcertArguments(ServerName, SessionName);
-		LaunchConcertServer(ServerName, SessionName);
+			if (DialogueResponse == FEditorFileUtils::EPromptReturnCode::PR_Cancelled)
+			{
+				return;
+			}
+		}
 	}
+ 
+	const FString EditorBinary = FPlatformProcess::ExecutablePath();
+	
+	const FString Project = FPaths::SetExtension(FPaths::Combine(FPaths::ProjectDir(), FApp::GetProjectName()),".uproject");
+	const FString Map = GetCurrentWorld()->GetCurrentLevel()->GetPackage()->GetFName().ToString();
 
 	for (const FString& Node : SelectedDisplayClusterConfigActorNodes)
 	{
@@ -322,7 +495,6 @@ void FDisplayClusterLaunchEditorModule::LaunchDisplayClusterProcess()
 		FString ConcatenatedConsoleCommands;
 		FString ConcatenatedDPCvars;
 		FString ConcatenatedLogCommands;
-
 		// Fullscreen/Windowed
 		if (UDisplayClusterConfigurationClusterNode** NodePtrPtr = ConfigDataToUse->Cluster->Nodes.Find(Node))
 		{
@@ -343,20 +515,17 @@ void FDisplayClusterLaunchEditorModule::LaunchDisplayClusterProcess()
 				}
 			}
 		}
-
 		// Unreal Insights support
 		if (ProjectSettings->bEnableUnrealInsights)
 		{
 			// Enable trace
 			ConcatenatedCommandLineArguments += " -trace";
-
 			if (ProjectSettings->bEnableStatNamedEvents)
 			{
 				ConcatenatedCommandLineArguments += " -statnamedevents";
 			}
-
 			// Override save directory if desired
-			if (ProjectSettings->bOverrideInsightsTraceFileSaveDirectory)
+			if (!ProjectSettings->ExplicitTraceFileSaveDirectory.Path.IsEmpty())
 			{
 				FString TraceFilePath =
 					ProjectSettings->ExplicitTraceFileSaveDirectory.Path / FDateTime::Now().ToString(TEXT("%Y%m%d_%H%M%S"));
@@ -364,53 +533,21 @@ void FDisplayClusterLaunchEditorModule::LaunchDisplayClusterProcess()
 					FString::Printf(TEXT(" -tracefile=%s "), *TraceFilePath);
 			}
 		}
-		
-		// Open a modal to prompt for save, if dirty. Yes = Save & Continue. No = Continue Without Saving. Cancel = Stop Opening Assets.
-		UPackage* PackageToSave = nullptr;
-        		
-		if (UWorld* World = GetCurrentWorld())
-		{
-			if (ULevel* Level = World->GetCurrentLevel())
-			{
-				PackageToSave = Level->GetPackage();
-			}
-		}
-
-		if (PackageToSave)
-		{
-			const FEditorFileUtils::EPromptReturnCode DialogueResponse =
-					FEditorFileUtils::PromptForCheckoutAndSave(
-						{PackageToSave},
-						true,
-						true,
-						LOCTEXT("SavePackagesTitle", "Save Packages"),
-						LOCTEXT("ConfirmOpenLevelFormat", "Do you want to save the current level?\n\nCancel to abort launch.\n")
-					);
-		
-			if (DialogueResponse == FEditorFileUtils::EPromptReturnCode::PR_Cancelled)
-			{
-				return;
-			}
-		}
 	
 		GetProjectSettingsArguments(ProjectSettings, ConcatenatedCommandLineArguments, ConcatenatedConsoleCommands, ConcatenatedDPCvars,
 					 ConcatenatedLogCommands);
 		
 		AddUdpMessagingArguments(ConcatenatedCommandLineArguments);
-
 		// Add nDisplay node information
 		ConcatenatedCommandLineArguments += " " + ConfigActorPath;
 		ConcatenatedCommandLineArguments += FString::Printf(TEXT(" -dc_node=\"%s\""), *Node);
-
 		// Add Multi-User params
 		if (!ConcertArguments.IsEmpty())
 		{
 			ConcatenatedCommandLineArguments += " " + ConcertArguments;
 		}
-
 		// Log file
 		const FString LogFileName = (ProjectSettings->LogFileName.IsEmpty() ? Node : ProjectSettings->LogFileName) + ".log";
-
 		const FString Params =
 			FString::Printf(
 				TEXT("\"%s\" -game \"%s\" Log=%s %s -ExecCmds=\"%s\" -DPCVars=\"%s\" -LogCmds=\"%s\""),
@@ -418,12 +555,9 @@ void FDisplayClusterLaunchEditorModule::LaunchDisplayClusterProcess()
 				*ConcatenatedConsoleCommands, *ConcatenatedDPCvars, *ConcatenatedLogCommands
 	
 		);
-
 		UE_LOG(LogDisplayClusterLaunchEditor, Log, TEXT("Full Command: %s %s"), *EditorBinary, *Params);
-
 		void* WritePipe = nullptr;
-
-		ActiveProcesses.Add(
+		ActiveDisplayClusterProcesses.Add(
 			FPlatformProcess::CreateProc(
 				*EditorBinary, *Params,
 				ProjectSettings->bCloseEditorOnLaunch,
@@ -436,12 +570,11 @@ void FDisplayClusterLaunchEditorModule::LaunchDisplayClusterProcess()
 
 void FDisplayClusterLaunchEditorModule::TerminateActiveDisplayClusterProcesses()
 {
-	for (FProcHandle& Process : ActiveProcesses)
+	for (FProcHandle& Process : ActiveDisplayClusterProcesses)
 	{
 		FPlatformProcess::TerminateProc(Process);
 		FPlatformProcess::CloseProc(Process);
 	}
-
 	RemoveTerminatedNodeProcesses();
 }
 
@@ -456,9 +589,7 @@ void FDisplayClusterLaunchEditorModule::RegisterToolbarItem()
 	UToolMenu* Menu = UToolMenus::Get()->ExtendMenu("LevelEditor.LevelEditorToolBar.User");
 	
 	RemoveToolbarItem();
-
 	FToolMenuSection& Section = Menu->AddSection("DisplayClusterLaunch");
-
 	const FToolMenuEntry DisplayClusterLaunchButton =
 		FToolMenuEntry::InitToolBarButton("DisplayClusterLaunchToolbarButton",
 			FUIAction(
@@ -472,7 +603,6 @@ void FDisplayClusterLaunchEditorModule::RegisterToolbarItem()
 				TAttribute<FSlateIcon>::FGetter::CreateRaw(this, &FDisplayClusterLaunchEditorModule::GetToolbarButtonIcon)
 			)
 		);
-
 	const FToolMenuEntry DisplayClusterLaunchComboButton = FToolMenuEntry::InitComboButton(
 		"DisplayClusterLaunchMenu",
 		FUIAction(),
@@ -482,27 +612,24 @@ void FDisplayClusterLaunchEditorModule::RegisterToolbarItem()
 		FSlateIcon(),
 		true //bInSimpleComboBox
 	);
-
 	Section.AddEntry(DisplayClusterLaunchButton);
 	Section.AddEntry(DisplayClusterLaunchComboButton);
 }
 
 FText FDisplayClusterLaunchEditorModule::GetToolbarButtonTooltipText()
 {
-	if (ActiveProcesses.Num() == 1)
+	if (ActiveDisplayClusterProcesses.Num() == 1)
 	{
 		return LOCTEXT("TerminateActiveProcess","Terminate active nDisplay process");
 	}
-	else if (ActiveProcesses.Num() > 1)
+	else if (ActiveDisplayClusterProcesses.Num() > 1)
 	{
-		return FText::Format(LOCTEXT("TerminateActiveProcessesFormat","Terminate {0} active nDisplay processes"), ActiveProcesses.Num());
+		return FText::Format(LOCTEXT("TerminateActiveProcessesFormat","Terminate {0} active nDisplay processes"), ActiveDisplayClusterProcesses.Num());
 	}
-
 	if (!SelectedDisplayClusterConfigActor.ResolveObject())
 	{
 		return LOCTEXT("GenericLaunchDisplayClusterProcessText_NoConfig", "Launch an nDisplay instance using the first Config Actor found in the current level and the first node found in that configuration.\n\nSet specific configurations and nodes using the overflow menu.");
 	}
-
 	if (SelectedDisplayClusterConfigActorNodes.Num() == 0)
 	{
 		return FText::Format(
@@ -510,9 +637,7 @@ FText FDisplayClusterLaunchEditorModule::GetToolbarButtonTooltipText()
 				FText::FromString(SelectedDisplayClusterConfigActor.GetAssetName())
 		);
 	}
-
 	FString ConfigActorName = SelectedDisplayClusterConfigActor.ResolveObject()->GetName();
-
 	const FString SplitTerm = "_C";
 	if (ConfigActorName.Contains(SplitTerm))
 	{
@@ -528,16 +653,15 @@ FText FDisplayClusterLaunchEditorModule::GetToolbarButtonTooltipText()
 FSlateIcon FDisplayClusterLaunchEditorModule::GetToolbarButtonIcon()
 {
 	RemoveTerminatedNodeProcesses();
-
 	return FSlateIcon(FAppStyle::Get().GetStyleSetName(),
-		ActiveProcesses.Num() > 0 ? "Icons.Toolbar.Stop" : "Icons.Toolbar.Play");
+		ActiveDisplayClusterProcesses.Num() > 0 ? "Icons.Toolbar.Stop" : "Icons.Toolbar.Play");
 }
 
 void FDisplayClusterLaunchEditorModule::OnClickToolbarButton()
 {
-	if (ActiveProcesses.Num() == 0)
+	if (ActiveDisplayClusterProcesses.Num() == 0)
 	{
-		LaunchDisplayClusterProcess();
+		TryLaunchDisplayClusterProcess();
 	}
 	else
 	{
@@ -578,7 +702,6 @@ FText FDisplayClusterLaunchEditorModule::GetSelectedNodesListText() const
 	
 		return FText::Format(FText::FromString(JoinedNodes), LOCTEXT("PrimaryNode", "Primary"));
 	}
-
 	return FText::GetEmpty();
 }
 
@@ -592,14 +715,11 @@ TArray<TWeakObjectPtr<ADisplayClusterRootActor>> FDisplayClusterLaunchEditorModu
 			CachedDisplayClusterActors.Add(RootActor);
 		}
 	}
-
 	bAreConfigsFoundInWorld = CachedDisplayClusterActors.Num() > 0;
-
 	CachedDisplayClusterActors.Sort([](const TWeakObjectPtr<ADisplayClusterRootActor>& A, const TWeakObjectPtr<ADisplayClusterRootActor>& B)
 	{
 		return A->GetActorLabel() < B->GetActorLabel();
 	});
-
 	return CachedDisplayClusterActors;
 }
 
@@ -615,7 +735,6 @@ void FDisplayClusterLaunchEditorModule::ApplyDisplayClusterConfigOverrides(UDisp
 		ConfigDataCopy->Scene = NewObject<UDisplayClusterConfigurationScene>(ConfigDataCopy, NAME_None,
 		RF_ArchetypeObject | RF_Public );
 	}
-
 	// A Primary Node should always be automatically selected, but this code preempts a crash. Normally we use the PN specified in the UI.
 	// If one is not specified in the UI, we check to see if the primary node specified
 	// in the original config is in our node array selection from the UI.
@@ -631,7 +750,6 @@ void FDisplayClusterLaunchEditorModule::ApplyDisplayClusterConfigOverrides(UDisp
 	{
 		ConfigDataCopy->Cluster->PrimaryNode.Id = SelectedDisplayClusterConfigActorPrimaryNode;
 	}
-
 	TMap<FString, UDisplayClusterConfigurationClusterNode*> ActiveNodes;
 	TMap<FString, UDisplayClusterConfigurationClusterNode*> NodesInConfig = ConfigDataCopy->Cluster->Nodes;
 	for (int32 NodeIndex = 0; NodeIndex < SelectedDisplayClusterConfigActorNodes.Num(); NodeIndex++)
@@ -640,9 +758,7 @@ void FDisplayClusterLaunchEditorModule::ApplyDisplayClusterConfigOverrides(UDisp
 		if (UDisplayClusterConfigurationClusterNode** Node = NodesInConfig.Find(NodeId))
 		{
 			ActiveNodes.Add(NodeId, *Node);
-
 			(*Node)->Host = "127.0.0.1";
-
 			// If we haven't specified a primary node and the config's primary node is not in our selection, use the first active node.
 			if (bIsPrimaryNodeUnset && !bIsConfigPrimaryNodeInActiveNodes && ActiveNodes.Num() == 1)
 			{
@@ -650,7 +766,6 @@ void FDisplayClusterLaunchEditorModule::ApplyDisplayClusterConfigOverrides(UDisp
 			}
 		}
 	}
-
 	ConfigDataCopy->Cluster->Nodes = ActiveNodes;
 }
 
@@ -659,13 +774,10 @@ void FDisplayClusterLaunchEditorModule::SetSelectedDisplayClusterConfigActor(ADi
 	if (SelectedActor)
 	{
 		const FSoftObjectPath AsSoftObjectPath(SelectedActor);
-
 		if (AsSoftObjectPath != SelectedDisplayClusterConfigActor)
 		{
 			SelectedDisplayClusterConfigActor = AsSoftObjectPath;
-
 			SelectedDisplayClusterConfigActorNodes.Empty();
-
 			SelectFirstNode(SelectedActor);
 		}
 	}
@@ -681,13 +793,11 @@ void FDisplayClusterLaunchEditorModule::ToggleDisplayClusterConfigActorNodeSelec
 	{
 		SelectedDisplayClusterConfigActorNodes.Add(InNodeName);
 	}
-
 	// Clear SelectedDisplayClusterConfigActorPrimaryNode if no nodes are selected
 	if (SelectedDisplayClusterConfigActorNodes.Num() == 0)
 	{
 		SelectedDisplayClusterConfigActorPrimaryNode = "";
 	}
-
 	// If a single node is selected, SelectedDisplayClusterConfigActorPrimaryNode must be this node
 	if (SelectedDisplayClusterConfigActorNodes.Num() == 1)
 	{
@@ -716,7 +826,6 @@ void FDisplayClusterLaunchEditorModule::SelectFirstNode(ADisplayClusterRootActor
 {
 	TArray<FString> NodeNames;
 	InConfig->GetConfigData()->Cluster->Nodes.GenerateKeyArray(NodeNames);
-
 	if (NodeNames.Num() == 0)
 	{
 		UE_LOG(LogDisplayClusterLaunchEditor, Error, TEXT("%hs: Unable to launch nDisplay because there are no nDisplay nodes in the selected nDisplay Config named '{0}'."), __FUNCTION__, *InConfig->GetActorNameOrLabel());
@@ -732,13 +841,12 @@ void FDisplayClusterLaunchEditorModule::SelectFirstNode(ADisplayClusterRootActor
 			}
 		)
 	);
-
 	if (SelectedDisplayClusterConfigActorNodes.Num() == 0)
 	{
 		const FString& NodeName = NodeNames[0];
 		UE_LOG(LogDisplayClusterLaunchEditor, Warning, TEXT("%hs: Selected nDisplay nodes were not found on the selected DisplayClusterRootActor. We will select the first valid node."), __FUNCTION__);
 		SelectedDisplayClusterConfigActorNodes.Add(NodeName);
-		UE_LOG(LogDisplayClusterLaunchEditor, Log, TEXT("%hs: Adding first valid node named '{0}' to selected nodes."), __FUNCTION__, *NodeName);
+		UE_LOG(LogDisplayClusterLaunchEditor, Log, TEXT("%hs: Adding first valid node named '%s' to selected nodes."), __FUNCTION__, *NodeName);
 	}
 }
 
@@ -747,9 +855,7 @@ TSharedRef<SWidget>  FDisplayClusterLaunchEditorModule::CreateToolbarMenuEntries
 	IAssetRegistry* AssetRegistry = IAssetRegistry::Get();
 	
 	FMenuBuilder MenuBuilder(false, nullptr);
-
 	TArray<TWeakObjectPtr<ADisplayClusterRootActor>> DisplayClusterConfigs = GetAllDisplayClusterConfigsInWorld();
-
 	MenuBuilder.BeginSection("DisplayClusterLaunch", LOCTEXT("DisplayClusterLauncher", "Launch nDisplay"));
 	{
 		MenuBuilder.AddMenuEntry(
@@ -757,22 +863,17 @@ TSharedRef<SWidget>  FDisplayClusterLaunchEditorModule::CreateToolbarMenuEntries
 			LOCTEXT("DisplayClusterLaunchLastNodeTooltip", "Launch the last node configuration."),
 			FSlateIcon(FAppStyle::Get().GetStyleSetName(), "Icons.Toolbar.Play"),
 			FUIAction(
-				FExecuteAction::CreateRaw(this, &FDisplayClusterLaunchEditorModule::LaunchDisplayClusterProcess),
+				FExecuteAction::CreateRaw(this, &FDisplayClusterLaunchEditorModule::TryLaunchDisplayClusterProcess),
 				FCanExecuteAction::CreateRaw(this, &FDisplayClusterLaunchEditorModule::DoesCurrentWorldHaveDisplayClusterConfig)
 			),
 			NAME_None
 		);
 	}
 	MenuBuilder.EndSection();
-
 	AddDisplayClusterLaunchConfigurations(AssetRegistry, MenuBuilder, DisplayClusterConfigs);
-
 	AddDisplayClusterLaunchNodes(AssetRegistry, MenuBuilder);
-
 	AddConsoleVariablesEditorAssetsToToolbarMenu(AssetRegistry, MenuBuilder);
-
 	AddOptionsToToolbarMenu(MenuBuilder);
-
 	return MenuBuilder.MakeWidget();
 }
 
@@ -792,7 +893,6 @@ void FDisplayClusterLaunchEditorModule::AddDisplayClusterLaunchConfigurations(IA
 					return Comparator.IsValid() && SelectedActor == Comparator.Get();
 				});
 			}
-
 			if (!bIsConfigActorValid)
 			{
 				SetSelectedDisplayClusterConfigActor(DisplayClusterConfigs[0].Get());
@@ -884,7 +984,6 @@ void FDisplayClusterLaunchEditorModule::AddDisplayClusterLaunchNodes(IAssetRegis
 			FOnGetContent::CreateLambda([this]()
 			{
 				FMenuBuilder NewMenuBuilder(false, nullptr);
-
 				NewMenuBuilder.AddSubMenu(
 					TAttribute<FText>::CreateLambda([this]()
 					{
@@ -932,9 +1031,7 @@ void FDisplayClusterLaunchEditorModule::AddDisplayClusterLaunchNodes(IAssetRegis
 					NAME_None,
 					EUserInterfaceActionType::None
 				);
-
 				NewMenuBuilder.AddSeparator();
-
 				if (ADisplayClusterRootActor* SelectedActor = Cast<ADisplayClusterRootActor>(SelectedDisplayClusterConfigActor.ResolveObject()))
 				{
 					TArray<FString> NodeNames;
@@ -963,12 +1060,10 @@ void FDisplayClusterLaunchEditorModule::AddDisplayClusterLaunchNodes(IAssetRegis
 					   );
 				   }
 				}
-
 				return NewMenuBuilder.MakeWidget();
 			}),
 			FSlateIcon(FDisplayClusterLaunchEditorStyle::Get().GetStyleSetName(),"Icons.DisplayClusterNode")
 		);
-
 		MenuBuilder.EndSection();
 	}
 }
@@ -978,12 +1073,11 @@ void FDisplayClusterLaunchEditorModule::AddConsoleVariablesEditorAssetsToToolbar
 {
 	TArray<FAssetData> FoundConsoleVariablesAssets;
 	AssetRegistry->GetAssetsByClass("ConsoleVariablesAsset", FoundConsoleVariablesAssets, true);
-
 	if (FoundConsoleVariablesAssets.Num())
 	{
 		MenuBuilder.BeginSection("DisplayClusterLaunchCvars", LOCTEXT("DisplayClusterLaunchCvars", "Console Variables"));
 		{
-			const FText ConsoleVariablesAssetTooltip = LOCTEXT("SelectConsoleVariablesAsset","Select Console Variables Asset");
+			const FText ConsoleVariablesAssetTooltip = LOCTEXT("SelectConsoleVariablesAssetFormat","Select Console Variables Asset");
 			
 			MenuBuilder.AddSubMenu(
 				TAttribute<FText>::Create([this](){ return FText::FromName(SelectedConsoleVariablesAssetName); }),
@@ -1041,7 +1135,6 @@ void FDisplayClusterLaunchEditorModule::AddOptionsToToolbarMenu(FMenuBuilder& Me
 			NAME_None,
 			EUserInterfaceActionType::Check
 		);
-
 		MenuBuilder.AddMenuEntry(
 			LOCTEXT("EnableUnrealInsightsLabel", "Enable Unreal Insights"),
 			LOCTEXT("EnableUnrealInsightsTooltip", "Enable Unreal Insights"),
@@ -1059,7 +1152,6 @@ void FDisplayClusterLaunchEditorModule::AddOptionsToToolbarMenu(FMenuBuilder& Me
 			NAME_None,
 			EUserInterfaceActionType::Check
 		);
-
 		MenuBuilder.AddMenuEntry(
 			LOCTEXT("CloseEditorOnLaunchLabel", "Close Editor on Launch"),
 			LOCTEXT("CloseEditorOnLaunchTooltip", "Close Editor on Launch"),
@@ -1077,7 +1169,6 @@ void FDisplayClusterLaunchEditorModule::AddOptionsToToolbarMenu(FMenuBuilder& Me
 			NAME_None,
 			EUserInterfaceActionType::Check
 		);
-
 		MenuBuilder.AddMenuEntry(
 			LOCTEXT("AdvancedSettingsLabel", "Advanced Settings..."),
 			LOCTEXT("AdvancedSettingsTooltip", "Open the nDisplay Launch Project Settings"),
@@ -1096,11 +1187,39 @@ bool FDisplayClusterLaunchEditorModule::GetConnectToMultiUser() const
 	return GetDefault<UDisplayClusterLaunchEditorProjectSettings>()->bConnectToMultiUser;
 }
 
+const FString& FDisplayClusterLaunchEditorModule::GetConcertServerName()
+{
+	// If the Cached name is changed after this point it will need to be cleared in order to run this code again 
+	if (ServerTrackingData.GeneratedMultiUserServerName.IsEmpty())
+	{
+		ServerTrackingData.GeneratedMultiUserServerName = AppendRandomNumbersToString("nDisplayLaunchServer");
+	}
+	return ServerTrackingData.GeneratedMultiUserServerName;
+}
+
+const FString& FDisplayClusterLaunchEditorModule::GetConcertSessionName()
+{
+	// If the Cached name is changed after this point it will need to be cleared in order to run this code again 
+	if (CachedConcertSessionName.IsEmpty())
+	{
+		const UDisplayClusterLaunchEditorProjectSettings* Settings = GetDefault<UDisplayClusterLaunchEditorProjectSettings>();
+		if (!Settings->ExplicitSessionName.IsEmpty())
+		{
+			CachedConcertSessionName = Settings->ExplicitSessionName;
+		}
+		else
+		{
+			CachedConcertSessionName = AppendRandomNumbersToString("nDisplayLaunchSession");
+		}
+	}
+	return CachedConcertSessionName;
+}
+
 void FDisplayClusterLaunchEditorModule::RemoveTerminatedNodeProcesses()
 {
-	ActiveProcesses.SetNum(
+	ActiveDisplayClusterProcesses.SetNum(
 		Algo::StableRemoveIf(
-			ActiveProcesses,
+			ActiveDisplayClusterProcesses,
 			[](FProcHandle& Handle)
 			{
 				return !FPlatformProcess::IsProcRunning(Handle);
diff --git a/Engine/Plugins/Editor/DisplayClusterLaunch/Source/DisplayClusterLaunchEditor/Private/DisplayClusterLaunchEditorStyle.cpp b/Engine/Plugins/Editor/DisplayClusterLaunch/Source/DisplayClusterLaunchEditor/Private/DisplayClusterLaunchEditorStyle.cpp
--- a/Engine/Plugins/Editor/DisplayClusterLaunch/Source/DisplayClusterLaunchEditor/Private/DisplayClusterLaunchEditorStyle.cpp
+++ b/Engine/Plugins/Editor/DisplayClusterLaunch/Source/DisplayClusterLaunchEditor/Private/DisplayClusterLaunchEditorStyle.cpp
@@ -70,7 +70,6 @@ const FVector2D Icon8x8(8.f, 8.f);
 FString FDisplayClusterLaunchEditorStyle::GetExternalPluginContent(const FString& PluginName, const FString& RelativePath, const ANSICHAR* Extension)
 {
 	FString ContentDir = IPluginManager::Get().FindPlugin(PluginName)->GetBaseDir() / RelativePath + Extension;
-	UE_LOG(LogTemp, Error, TEXT("%hs: Content = %s"), __FUNCTION__, *ContentDir);
 	return ContentDir;
 }
 
diff --git a/Engine/Plugins/Editor/DisplayClusterLaunch/Source/DisplayClusterLaunchEditor/Public/DisplayClusterLaunchEditorModule.h b/Engine/Plugins/Editor/DisplayClusterLaunch/Source/DisplayClusterLaunchEditor/Public/DisplayClusterLaunchEditorModule.h
--- a/Engine/Plugins/Editor/DisplayClusterLaunch/Source/DisplayClusterLaunchEditor/Public/DisplayClusterLaunchEditorModule.h
+++ b/Engine/Plugins/Editor/DisplayClusterLaunch/Source/DisplayClusterLaunchEditor/Public/DisplayClusterLaunchEditorModule.h
@@ -2,6 +2,9 @@
 
 #pragma once
 
+#include "ConcertMessageData.h"
+#include "IConcertModule.h"
+
 #include "Containers/Array.h"
 #include "Modules/ModuleInterface.h"
 #include "Templates/SharedPointer.h"
@@ -27,10 +30,19 @@ class FDisplayClusterLaunchEditorModule : public IModuleInterface
 {
 public:
 
-	enum class ELaunchState
+	enum class EConcertServerRequestStatus
+	{
+		None,
+		ShutdownRequested,
+		LaunchRequested,
+		ReuseExisting
+	};
+
+	struct FDisplayClusterLaunchMultiUserServerTrackingData
 	{
-		NotLaunched,
-		Launched
+		FString GeneratedMultiUserServerName;
+		FConcertServerInfo MultiUserServerInfo;
+		FProcHandle MultiUserServerHandle;
 	};
 	
 	static FDisplayClusterLaunchEditorModule& Get();
@@ -42,13 +54,27 @@ public:
 	
 	static void OpenProjectSettings();
 
-	void LaunchDisplayClusterProcess();
+	/**
+	 * Performs a sanity check to determine if nDisplay can be launched
+	 * and if so, gets Concert parameters async if Concert is enabled for the launch.
+	 * Ultimately, LaunchDisplayClusterProcess will be called when Concert parameters are skipped or returned.
+	 */
+	void TryLaunchDisplayClusterProcess();
 	void TerminateActiveDisplayClusterProcesses();
 
 private:
 
 	void OnFEngineLoopInitComplete();
 
+	void LaunchConcertServer();
+
+	void FindOrLaunchConcertServer();
+	void OnServersAssumedReady();
+	void FindAppropriateServer();
+	void ConnectToSession();
+	
+	void LaunchDisplayClusterProcess();
+
 	void RegisterToolbarItem();
 	FText GetToolbarButtonTooltipText();
 	FSlateIcon GetToolbarButtonIcon();
@@ -79,6 +105,8 @@ private:
 	void AddOptionsToToolbarMenu(FMenuBuilder& MenuBuilder);
 
 	bool GetConnectToMultiUser() const;
+	const FString& GetConcertServerName();
+	const FString& GetConcertSessionName(); 
 
 	void RemoveTerminatedNodeProcesses();
 
@@ -90,5 +118,10 @@ private:
 	
 	FName SelectedConsoleVariablesAssetName = NAME_None;
 
-	TArray<FProcHandle> ActiveProcesses; 
+	FDisplayClusterLaunchMultiUserServerTrackingData ServerTrackingData;
+	FString CachedConcertSessionName;
+	
+	TArray<FProcHandle> ActiveDisplayClusterProcesses;
+
+	EConcertServerRequestStatus ConcertServerRequestStatus = EConcertServerRequestStatus::None;
 };
diff --git a/Engine/Plugins/Editor/DisplayClusterLaunch/Source/DisplayClusterLaunchEditor/Public/DisplayClusterLaunchEditorProjectSettings.h b/Engine/Plugins/Editor/DisplayClusterLaunch/Source/DisplayClusterLaunchEditor/Public/DisplayClusterLaunchEditorProjectSettings.h
--- a/Engine/Plugins/Editor/DisplayClusterLaunch/Source/DisplayClusterLaunchEditor/Public/DisplayClusterLaunchEditorProjectSettings.h
+++ b/Engine/Plugins/Editor/DisplayClusterLaunch/Source/DisplayClusterLaunchEditor/Public/DisplayClusterLaunchEditorProjectSettings.h
@@ -93,34 +93,11 @@ public:
 	UPROPERTY(Config, EditAnywhere, Category="nDisplay Launch Settings|Multi-User")
 	bool bConnectToMultiUser = false;
 
-	/**
-	 * If true, a server name will be automatically generated for you when connecting to multi-user.
-	 * If false, the text in ExplicitServerName will be used instead.
-	 * If ExplicitServerName is empty, a name will be generated whether this setting is true or false.
-	 */
-	UPROPERTY(Config, EditAnywhere, Category="nDisplay Launch Settings|Multi-User", meta = (EditCondition="bConnectToMultiUser"))
-	bool bAutoGenerateServerName = true;
-
-	/**
-	 * A specific server name to use when connecting to multi-user if bAutoGenerateServerName is false.
-	 * If left empty, a name will be generated whether bAutoGenerateServerName is true or false.
-	 */
-	UPROPERTY(Config, EditAnywhere, Category="nDisplay Launch Settings|Multi-User", meta = (EditCondition="bConnectToMultiUser && !bAutoGenerateServerName"))
-	FString ExplicitServerName;
-
-	/**
-	 * If true, a session name will be automatically generated for you when connecting to multi-user.
-	 * If false, the text in ExplicitSessionName will be used instead.
-	 * If ExplicitSessionName is empty, a name will be generated whether this setting is true or false.
-	 */
-	UPROPERTY(Config, EditAnywhere, Category="nDisplay Launch Settings|Multi-User", meta = (EditCondition="bConnectToMultiUser"))
-	bool bAutoGenerateSessionName = true;
-
 	/**
 	 * A specific session name to use when connecting to multi-user if bAutoGenerateSessionName is false.
-	 * If left empty, a name will be generated whether bAutoGenerateSessionName is true or false.
+	 * If left empty, a name will be generated.
 	 */
-	UPROPERTY(Config, EditAnywhere, Category="nDisplay Launch Settings|Multi-User", meta = (EditCondition="bConnectToMultiUser && !bAutoGenerateSessionName"))
+	UPROPERTY(Config, EditAnywhere, Category="nDisplay Launch Settings|Multi-User", meta = (EditCondition="bConnectToMultiUser"))
 	FString ExplicitSessionName;
 
 	/**
@@ -136,17 +113,11 @@ public:
 	bool bEnableStatNamedEvents = false;
 
 	/**
-	 * If true, the Unreal Insights Trace Files will be saved to the path specified in ExplicitTraceFileSaveDirectory without needing to run Unreal Insights.
-	 * If false, Unreal Insights will connect to localhost (this computer) instead and you'll need to ensure Unreal Insights is launched.
+	 * Where to store the Unreal Insights Trace Files.
+	 * If empty, Unreal Insights will connect to localhost (this computer) instead and you'll need to ensure Unreal Insights is launched.
 	 * To specify a socket to connect to, use CommandLineArguments.
 	 */
 	UPROPERTY(Config, EditAnywhere, Category="nDisplay Launch Settings|Unreal Insights", meta = (EditCondition="bEnableUnrealInsights"))
-	bool bOverrideInsightsTraceFileSaveDirectory = false;
-
-	/**
-	 * Where to store the Unreal Insights Trace Files if bOverrideInsightsTraceFileDirectory is true.
-	 */
-	UPROPERTY(Config, EditAnywhere, Category="nDisplay Launch Settings|Unreal Insights", meta = (EditCondition="bEnableUnrealInsights && bOverrideInsightsTraceFileSaveDirectory"))
 	FDirectoryPath ExplicitTraceFileSaveDirectory;
 
 	/**
