Move logic for config parse after file save.

diff --git a/Engine/Plugins/Editor/DisplayClusterLaunch/Source/DisplayClusterLaunchEditor/Private/DisplayClusterLaunchEditorModule.cpp b/Engine/Plugins/Editor/DisplayClusterLaunch/Source/DisplayClusterLaunchEditor/Private/DisplayClusterLaunchEditorModule.cpp
--- a/Engine/Plugins/Editor/DisplayClusterLaunch/Source/DisplayClusterLaunchEditor/Private/DisplayClusterLaunchEditorModule.cpp
+++ b/Engine/Plugins/Editor/DisplayClusterLaunch/Source/DisplayClusterLaunchEditor/Private/DisplayClusterLaunchEditorModule.cpp
@@ -420,32 +420,7 @@ void FDisplayClusterLaunchEditorModule::TryLaunchDisplayClusterProcess()
 }
 
 void FDisplayClusterLaunchEditorModule::LaunchDisplayClusterProcess()
-{
-	const UDisplayClusterLaunchEditorProjectSettings* ProjectSettings = GetDefault<UDisplayClusterLaunchEditorProjectSettings>();
-	UDisplayClusterConfigurationData* ConfigDataToUse = nullptr;
-	FString ConfigActorPath;
-	
-	// If it's valid we need to check the selected nodes against the current config. If they don't exist, we need to get the first one.
-	if (const ADisplayClusterRootActor* ConfigActor = Cast<ADisplayClusterRootActor>(SelectedDisplayClusterConfigActor.ResolveObject()))
-	{
-		// Duplicate existing config data so we can make non-destructive edits
-		ConfigDataToUse = DuplicateObject(ConfigActor->GetConfigData(), GetTransientPackage());
-		ApplyDisplayClusterConfigOverrides(ConfigDataToUse);
-		const FString FilePath = FPaths::Combine(FPaths::ProjectSavedDir(), "Temp.ndisplay");
-		if (!ensureAlways(IDisplayClusterConfiguration::Get().SaveConfig(ConfigDataToUse, FilePath)))
-		{
-			UE_LOG(LogDisplayClusterLaunchEditor, Error, TEXT("%hs: Unable to launch nDisplay because the selected nDisplay Configuration could not be saved to a .ndisplay file. See the log for more information."), __FUNCTION__);
-			return;
-		}
-		
-		ConfigActorPath = FString::Printf(TEXT("-dc_cfg=\"%s\""), *FilePath);
-	}
-	else
-	{
-		UE_LOG(LogDisplayClusterLaunchEditor, Error, TEXT("%hs: Unable to launch nDisplay because the selected nDisplay Config Actor could not be resolved or does not exist in the current level."), __FUNCTION__);
-		return;
-	}
-	
+{		
 	UE_LOG(LogDisplayClusterLaunchEditor, Log, TEXT("%hs: Launching nDisplay processes..."), __FUNCTION__);
 	
 	FString ConcertArguments;
@@ -458,8 +433,8 @@ void FDisplayClusterLaunchEditorModule::LaunchDisplayClusterProcess()
 	{
 		// Open a modal to prompt for save, if dirty. Yes = Save & Continue. No = Continue Without Saving. Cancel = Stop Opening Assets.
 		UPackage* PackageToSave = nullptr;
-        		
-		if (UWorld* World = GetCurrentWorld())
+
+                if (UWorld* World = GetCurrentWorld())
 		{
 			if (ULevel* Level = World->GetCurrentLevel())
 			{
@@ -476,14 +451,39 @@ void FDisplayClusterLaunchEditorModule::LaunchDisplayClusterProcess()
 						LOCTEXT("SavePackagesTitle", "Save Packages"),
 						LOCTEXT("ConfirmOpenLevelFormat", "Do you want to save the current level?\n\nCancel to abort launch.\n")
 					);
-		
+
 			if (DialogueResponse == FEditorFileUtils::EPromptReturnCode::PR_Cancelled)
 			{
 				return;
 			}
 		}
 	}
- 
+
+	const UDisplayClusterLaunchEditorProjectSettings* ProjectSettings = GetDefault<UDisplayClusterLaunchEditorProjectSettings>();
+	UDisplayClusterConfigurationData* ConfigDataToUse = nullptr;
+	FString ConfigActorPath;
+
+	// If it's valid we need to check the selected nodes against the current config. If they don't exist, we need to get the first one.
+	if (const ADisplayClusterRootActor* ConfigActor = Cast<ADisplayClusterRootActor>(SelectedDisplayClusterConfigActor.ResolveObject()))
+	{
+		// Duplicate existing config data so we can make non-destructive edits
+		ConfigDataToUse = DuplicateObject(ConfigActor->GetConfigData(), GetTransientPackage());
+		ApplyDisplayClusterConfigOverrides(ConfigDataToUse);
+		const FString FilePath = FPaths::Combine(FPaths::ProjectSavedDir(), "Temp.ndisplay");
+		if (!ensureAlways(IDisplayClusterConfiguration::Get().SaveConfig(ConfigDataToUse, FilePath)))
+		{
+			UE_LOG(LogDisplayClusterLaunchEditor, Error, TEXT("%hs: Unable to launch nDisplay because the selected nDisplay Configuration could not be saved to a .ndisplay file. See the log for more information."), __FUNCTION__);
+			return;
+		}
+	
+		ConfigActorPath = FString::Printf(TEXT("-dc_cfg=\"%s\""), *FilePath);
+	}
+	else
+	{
+		UE_LOG(LogDisplayClusterLaunchEditor, Error, TEXT("%hs: Unable to launch nDisplay because the selected nDisplay Config Actor could not be resolved or does not exist in the current level."), __FUNCTION__);
+		return;
+	}
+
 	const FString EditorBinary = FPlatformProcess::ExecutablePath();
 	
 	const FString Project = FPaths::SetExtension(FPaths::Combine(FPaths::ProjectDir(), FApp::GetProjectName()),".uproject");
