nDisplay Launch:

-Update Enum path when using EnumToString method

diff --git a/Engine/Plugins/Editor/DisplayClusterLaunch/Source/DisplayClusterLaunchEditor/Private/DisplayClusterLaunchEditorModule.cpp b/Engine/Plugins/Editor/DisplayClusterLaunch/Source/DisplayClusterLaunchEditor/Private/DisplayClusterLaunchEditorModule.cpp
--- a/Engine/Plugins/Editor/DisplayClusterLaunch/Source/DisplayClusterLaunchEditor/Private/DisplayClusterLaunchEditorModule.cpp
+++ b/Engine/Plugins/Editor/DisplayClusterLaunch/Source/DisplayClusterLaunchEditor/Private/DisplayClusterLaunchEditorModule.cpp
@@ -37,14 +37,16 @@ void CloseAllMenus()
 	FSlateApplication::Get().DismissAllMenus();
 }
 
-FString EnumToString(const TCHAR* EnumName, const int32 EnumValue)
+FString EnumToString(const FString EnumName, const int32 EnumValue)
 {
-	const UEnum* EnumPtr = FindObject<UEnum>(nullptr, EnumName, true);
+	const FString EnumPath = "/Script/DisplayClusterLaunchEditor." + EnumName;
+	const UEnum* EnumPtr = FindObject<UEnum>(nullptr, *EnumPath, true);
 	
 	if (!EnumPtr)
 	{
 		return LOCTEXT("EnumNotFound", "Enum not found").ToString();
 	}
+	
 	return EnumPtr->GetNameStringByIndex(EnumValue);
 }
 
@@ -134,7 +136,7 @@ void GetProjectSettingsArguments(const UDisplayClusterLaunchEditorProjectSetting
 			}
 			ConcatenatedLogCommands += FString::Printf(TEXT("%s %s, "),
 			                                           *LoggingConstruct.Category.ToString(),
-			                                           *EnumToString(TEXT("/Script/DisplayClusterLaunchEditor.EDisplayClusterLaunchLogVerbosity"), (int32)LoggingConstruct.VerbosityLevel.GetValue()));
+			                                           *EnumToString("EDisplayClusterLaunchLogVerbosity", (int32)LoggingConstruct.VerbosityLevel.GetValue()));
 		}
 		// Remove whitespace
 		ConcatenatedLogCommands.TrimStartAndEndInline();
