nDisplay Launch:

-Add console variables asset support (need only specify the asset in settings)
-Add menu to create config/place default nDisplay actors
-Fix for: File path cannot be found when using a map that is first saved by launch prompt and can lead to node crash
-Fix for: Primary node does not update when switching config and crashes on launch

diff --git a/Engine/Plugins/Editor/DisplayClusterLaunch/Source/DisplayClusterLaunchEditor/Private/DisplayClusterLaunchEditorModule.cpp b/Engine/Plugins/Editor/DisplayClusterLaunch/Source/DisplayClusterLaunchEditor/Private/DisplayClusterLaunchEditorModule.cpp
--- a/Engine/Plugins/Editor/DisplayClusterLaunch/Source/DisplayClusterLaunchEditor/Private/DisplayClusterLaunchEditorModule.cpp
+++ b/Engine/Plugins/Editor/DisplayClusterLaunch/Source/DisplayClusterLaunchEditor/Private/DisplayClusterLaunchEditorModule.cpp
@@ -22,6 +22,7 @@
 #include "Editor.h"
 #include "EngineUtils.h"
 #include "FileHelpers.h"
+#include "IPlacementModeModule.h"
 #include "Engine/GameEngine.h"
 #include "GenericPlatform/GenericPlatformMisc.h"
 #include "ISettingsModule.h"
@@ -29,6 +30,9 @@
 #include "Misc/ConfigCacheIni.h"
  
 #include "ToolMenus.h"
+#include "ActorFactories/ActorFactoryBlueprint.h"
+#include "ConsoleVariablesEditorRuntime/Public/ConsoleVariablesAsset.h"
+#include "Interfaces/IPluginManager.h"
 
 #define LOCTEXT_NAMESPACE "FDisplayClusterLaunchEditorModule"
 
@@ -105,46 +109,6 @@ void FDisplayClusterLaunchEditorModule::OpenProjectSettings()
 		.ShowViewer("Project", "Plugins", "nDisplay Launch");
 }
 
-void GetProjectSettingsArguments(const UDisplayClusterLaunchEditorProjectSettings* ProjectSettings, FString& ConcatenatedCommandLineArguments, FString& ConcatenatedConsoleCommands, FString& ConcatenatedDPCvars, FString& ConcatenatedLogCommands)
-{
-	{
-		for (const FString& CommandLineArgument : ProjectSettings->CommandLineArguments)
-		{
-			if (CommandLineArgument.IsEmpty())
-			{
-				continue;
-			}
-			ConcatenatedCommandLineArguments += FString::Printf(TEXT(" -%s "), *CommandLineArgument);
-		}
-		// Remove whitespace
-		ConcatenatedCommandLineArguments.TrimStartAndEndInline();
-	}
-	if (ProjectSettings->AdditionalConsoleCommands.Num() > 0)
-	{
-		ConcatenatedConsoleCommands += FString::Join(ProjectSettings->AdditionalConsoleCommands, TEXT(","));
-	}
-	if (ProjectSettings->AdditionalConsoleVariables.Num() > 0)
-	{
-		ConcatenatedDPCvars += FString::Join(ProjectSettings->AdditionalConsoleVariables, TEXT(","));
-	}
-	{
-		for (const FDisplayClusterLaunchLoggingConstruct& LoggingConstruct : ProjectSettings->Logging)
-		{
-			if (LoggingConstruct.Category.IsNone())
-			{
-				continue;
-			}
-			ConcatenatedLogCommands += FString::Printf(TEXT("%s %s, "),
-			                                           *LoggingConstruct.Category.ToString(),
-			                                           *EnumToString("EDisplayClusterLaunchLogVerbosity", (int32)LoggingConstruct.VerbosityLevel.GetValue()));
-		}
-		// Remove whitespace
-		ConcatenatedLogCommands.TrimStartAndEndInline();
-		// Remove last comma
-		ConcatenatedLogCommands = ConcatenatedLogCommands.LeftChop(1);
-	}
-}
-
 bool AddUdpMessagingArguments(FString& ConcatenatedArguments)
 {
 	// Get from reflection because UUdpMessagingSettings does not export its API
@@ -424,42 +388,40 @@ void FDisplayClusterLaunchEditorModule::TryLaunchDisplayClusterProcess()
 void FDisplayClusterLaunchEditorModule::LaunchDisplayClusterProcess()
 {		
 	UE_LOG(LogDisplayClusterLaunchEditor, Log, TEXT("%hs: Launching nDisplay processes..."), __FUNCTION__);
-	
-	FString ConcertArguments;
 
-	if (GetConnectToMultiUser())
+	// Open a modal to prompt for save, if dirty. Yes = Save & Continue. No = Continue Without Saving. Cancel = Stop Opening Assets.
+	UPackage* PackageToSave = nullptr;
+
+	if (UWorld* World = GetCurrentWorld())
 	{
-		ConcertArguments = GetConcertArguments(GetConcertServerName(), GetConcertSessionName());
+		if (ULevel* Level = World->GetCurrentLevel())
+		{
+			PackageToSave = Level->GetPackage();
+		}
 	}
-	else
+	if (PackageToSave)
 	{
-		// Open a modal to prompt for save, if dirty. Yes = Save & Continue. No = Continue Without Saving. Cancel = Stop Opening Assets.
-		UPackage* PackageToSave = nullptr;
+		const FEditorFileUtils::EPromptReturnCode DialogueResponse =
+				FEditorFileUtils::PromptForCheckoutAndSave(
+					{PackageToSave},
+					true,
+					true,
+					LOCTEXT("SavePackagesTitle", "Save Packages"),
+					LOCTEXT("ConfirmOpenLevelFormat", "Do you want to save the current level?\n\nCancel to abort launch.\n")
+				);
 
-                if (UWorld* World = GetCurrentWorld())
+		if (DialogueResponse == FEditorFileUtils::EPromptReturnCode::PR_Cancelled)
 		{
-			if (ULevel* Level = World->GetCurrentLevel())
-			{
-				PackageToSave = Level->GetPackage();
-			}
-		}
-		if (PackageToSave)
-		{
-			const FEditorFileUtils::EPromptReturnCode DialogueResponse =
-					FEditorFileUtils::PromptForCheckoutAndSave(
-						{PackageToSave},
-						true,
-						true,
-						LOCTEXT("SavePackagesTitle", "Save Packages"),
-						LOCTEXT("ConfirmOpenLevelFormat", "Do you want to save the current level?\n\nCancel to abort launch.\n")
-					);
-
-			if (DialogueResponse == FEditorFileUtils::EPromptReturnCode::PR_Cancelled)
-			{
-				return;
-			}
+			return;
 		}
 	}
+	
+	FString ConcertArguments;
+
+	if (GetConnectToMultiUser())
+	{
+		ConcertArguments = GetConcertArguments(GetConcertServerName(), GetConcertSessionName());
+	}
 
 	const UDisplayClusterLaunchEditorProjectSettings* ProjectSettings = GetDefault<UDisplayClusterLaunchEditorProjectSettings>();
 	UDisplayClusterConfigurationData* ConfigDataToUse = nullptr;
@@ -584,6 +546,7 @@ void FDisplayClusterLaunchEditorModule::OnFEngineLoopInitComplete()
 {
 	RegisterProjectSettings();
 	RegisterToolbarItem();
+	RegisterPlacementModeItemsIfTheyExist();
 }
 
 void FDisplayClusterLaunchEditorModule::RegisterToolbarItem()
@@ -693,6 +656,81 @@ void FDisplayClusterLaunchEditorModule::RegisterProjectSettings() const
 	}
 }
 
+void FDisplayClusterLaunchEditorModule::RegisterPlacementModeItemsIfTheyExist()
+{
+	if (GEditor)
+	{
+		TSharedPtr<IPlugin> Plugin = IPluginManager::Get().FindPlugin(TEXT("DisplayClusterLaunch"));
+		check(Plugin.IsValid());
+
+		IAssetRegistry* AssetRegistry = IAssetRegistry::Get();
+
+		TArray<FAssetData> OutAssets;
+		const FString PlaceablePath = Plugin->GetContentDir();
+		
+		FString PackagePath;
+		FString* OutFailureReason = nullptr;
+		FPackageName::TryConvertFilenameToLongPackageName(PlaceablePath, PackagePath, OutFailureReason);
+
+		AssetRegistry->ScanPathsSynchronous({PlaceablePath});
+		AssetRegistry->GetAssetsByPath(*PackagePath, OutAssets);
+
+		if (OutAssets.Num() > 0)
+		{
+			if (const FPlacementCategoryInfo* Info = GetDisplayClusterPlacementCategoryInfo())
+			{
+				for (const FAssetData& Asset : OutAssets)
+				{
+					TSharedRef<FPlaceableItem> Item =
+						MakeShared<FPlaceableItem>(
+							*UActorFactoryBlueprint::StaticClass(),
+							Asset
+						);
+
+					Item->DisplayName =
+						FText::FromString(
+							Item->DisplayName.ToString().Replace(TEXT("_"), TEXT(" "))
+						);
+					
+					IPlacementModeModule::Get().RegisterPlaceableItem(Info->UniqueHandle, Item);
+				}
+			}
+		}
+	}
+}
+
+const FPlacementCategoryInfo* FDisplayClusterLaunchEditorModule::GetDisplayClusterPlacementCategoryInfo()
+{
+	if (GEditor)
+	{
+		const IPlacementModeModule& PlacementModeModule = IPlacementModeModule::Get();
+		const FName PlacementModeCategoryHandle = TEXT("nDisplay");
+
+		if (const FPlacementCategoryInfo* RegisteredInfo = PlacementModeModule.GetRegisteredPlacementCategory(PlacementModeCategoryHandle))
+		{
+			return RegisteredInfo;
+		}
+		else
+		{
+			FPlacementCategoryInfo Info(
+				LOCTEXT("DisplayClusterCategoryName", "nDisplay"),
+				FSlateIcon(FDisplayClusterLaunchEditorStyle::Get().GetStyleSetName(), "Icons.DisplayCluster"),
+				PlacementModeCategoryHandle,
+				TEXT("nDisplay"),
+				25
+			);
+
+			IPlacementModeModule::Get().RegisterPlacementCategory(Info);
+
+			// This will return nullptr if the Register above failed so we don't need to explicitly check
+			// RegisterPlacementCategory's return value.
+			return PlacementModeModule.GetRegisteredPlacementCategory(PlacementModeCategoryHandle);
+		}
+	}
+
+	return nullptr;
+}
+
 FText FDisplayClusterLaunchEditorModule::GetSelectedNodesListText() const
 {
 	if (SelectedDisplayClusterConfigActorNodes.Num() > 0)
@@ -707,6 +745,72 @@ FText FDisplayClusterLaunchEditorModule::GetSelectedNodesListText() const
 	return FText::GetEmpty();
 }
 
+void FDisplayClusterLaunchEditorModule::GetProjectSettingsArguments(
+	const UDisplayClusterLaunchEditorProjectSettings* ProjectSettings, FString& ConcatenatedCommandLineArguments,
+	FString& ConcatenatedConsoleCommands, FString& ConcatenatedDPCvars, FString& ConcatenatedLogCommands)
+{
+	{
+		for (const FString& CommandLineArgument : ProjectSettings->CommandLineArguments)
+		{
+			if (!CommandLineArgument.IsEmpty())
+			{
+				ConcatenatedCommandLineArguments += FString::Printf(TEXT(" -%s "), *CommandLineArgument);
+			}
+		}
+		// Remove whitespace
+		ConcatenatedCommandLineArguments.TrimStartAndEndInline();
+	}
+
+	// Add console commands
+	TArray<FString> CollectedCommands;
+
+	auto ConsoleVariablesAssetLambda = [](FSoftObjectPath InAssetSoftPath, TArray<FString> OutCollectedCommands)
+	{
+		if (InAssetSoftPath.IsValid())
+		{
+			if (TObjectPtr<UConsoleVariablesAsset> CvarAsset = Cast<UConsoleVariablesAsset>(InAssetSoftPath.TryLoad()))
+			{
+				OutCollectedCommands.Append(CvarAsset->GetSavedCommandsAsStringArray());
+			}
+			else
+			{
+				UE_LOG(LogDisplayClusterLaunchEditor, Warning, TEXT("%hs: The specified Console Variables Asset was not valid: %s."), __FUNCTION__, *InAssetSoftPath.GetAssetName());
+			}
+		}
+	};
+
+	ConsoleVariablesAssetLambda(ProjectSettings->ConsoleVariablesPreset, CollectedCommands);
+	ConsoleVariablesAssetLambda(SelectedAdditionalConsoleVariablesAsset, CollectedCommands);
+
+	// Additional Console Commands
+	CollectedCommands.Append(ProjectSettings->AdditionalConsoleCommands.Array());
+	
+	ConcatenatedConsoleCommands += FString::Join(CollectedCommands, TEXT(","));
+
+	// Additional Console Variables
+	ConcatenatedDPCvars += FString::Join(ProjectSettings->AdditionalConsoleVariables, TEXT(","));
+
+	// Logging
+	{
+		for (const FDisplayClusterLaunchLoggingConstruct& LoggingConstruct : ProjectSettings->Logging)
+		{
+			if (!LoggingConstruct.Category.IsNone())
+			{
+				ConcatenatedLogCommands +=
+					FString::Printf(TEXT("%s %s, "),
+						*LoggingConstruct.Category.ToString(),
+						*EnumToString("EDisplayClusterLaunchLogVerbosity", (int32)LoggingConstruct.VerbosityLevel.GetValue()));
+			}
+		}
+		
+		// Remove whitespace
+		ConcatenatedLogCommands.TrimStartAndEndInline();
+		// Remove last comma
+		ConcatenatedLogCommands = ConcatenatedLogCommands.LeftChop(1);
+	}
+}
+
+
 TArray<TWeakObjectPtr<ADisplayClusterRootActor>> FDisplayClusterLaunchEditorModule::GetAllDisplayClusterConfigsInWorld()
 {
 	TArray<TWeakObjectPtr<ADisplayClusterRootActor>> CachedDisplayClusterActors;
@@ -795,11 +899,13 @@ void FDisplayClusterLaunchEditorModule::ToggleDisplayClusterConfigActorNodeSelec
 	{
 		SelectedDisplayClusterConfigActorNodes.Add(InNodeName);
 	}
+	
 	// Clear SelectedDisplayClusterConfigActorPrimaryNode if no nodes are selected
 	if (SelectedDisplayClusterConfigActorNodes.Num() == 0)
 	{
 		SelectedDisplayClusterConfigActorPrimaryNode = "";
 	}
+	
 	// If a single node is selected, SelectedDisplayClusterConfigActorPrimaryNode must be this node
 	if (SelectedDisplayClusterConfigActorNodes.Num() == 1)
 	{
@@ -814,13 +920,13 @@ bool FDisplayClusterLaunchEditorModule::IsDisplayClusterConfigActorNodeSelected(
 
 void FDisplayClusterLaunchEditorModule::SetSelectedConsoleVariablesAsset(const FAssetData InConsoleVariablesAsset)
 {
-	if (SelectedConsoleVariablesAssetName == InConsoleVariablesAsset.AssetName)
+	if (SelectedAdditionalConsoleVariablesAsset.GetAssetName().Equals(InConsoleVariablesAsset.AssetName.ToString()))
 	{
-		SelectedConsoleVariablesAssetName = NAME_None;
+		SelectedAdditionalConsoleVariablesAsset.Reset();
 	}
 	else
 	{
-		SelectedConsoleVariablesAssetName = InConsoleVariablesAsset.AssetName;
+		SelectedAdditionalConsoleVariablesAsset = InConsoleVariablesAsset.ToSoftObjectPath();
 	}
 }
 
@@ -843,6 +949,7 @@ void FDisplayClusterLaunchEditorModule::SelectFirstNode(ADisplayClusterRootActor
 			}
 		)
 	);
+	
 	if (SelectedDisplayClusterConfigActorNodes.Num() == 0)
 	{
 		const FString& NodeName = NodeNames[0];
@@ -850,6 +957,9 @@ void FDisplayClusterLaunchEditorModule::SelectFirstNode(ADisplayClusterRootActor
 		SelectedDisplayClusterConfigActorNodes.Add(NodeName);
 		UE_LOG(LogDisplayClusterLaunchEditor, Log, TEXT("%hs: Adding first valid node named '%s' to selected nodes."), __FUNCTION__, *NodeName);
 	}
+
+	// Set the primary node to be the first in the list
+	SelectedDisplayClusterConfigActorPrimaryNode = SelectedDisplayClusterConfigActorNodes[0];
 }
 
 TSharedRef<SWidget>  FDisplayClusterLaunchEditorModule::CreateToolbarMenuEntries()
@@ -1075,14 +1185,23 @@ void FDisplayClusterLaunchEditorModule::AddConsoleVariablesEditorAssetsToToolbar
 {
 	TArray<FAssetData> FoundConsoleVariablesAssets;
 	AssetRegistry->GetAssetsByClass(FTopLevelAssetPath(TEXT("/Script/ConsoleVariablesEditor"), TEXT("ConsoleVariablesAsset")), FoundConsoleVariablesAssets, true);
+	
 	if (FoundConsoleVariablesAssets.Num())
 	{
-		MenuBuilder.BeginSection("DisplayClusterLaunchCvars", LOCTEXT("DisplayClusterLaunchCvars", "Console Variables"));
+		MenuBuilder.BeginSection("DisplayClusterLaunchCvars", LOCTEXT("DisplayClusterLaunch_AdditionalConsoleVariablesAsset", "Additional Console Variables Asset"));
 		{
-			const FText ConsoleVariablesAssetTooltip = LOCTEXT("SelectConsoleVariablesAssetFormat","Select Console Variables Asset");
+			const FText ConsoleVariablesAssetTooltip = LOCTEXT("SelectConsoleVariablesAssetTooltip","Select Additional Console Variables Asset. This asset will be in addition to and its commands will be executed after the one defined in project settings.");
 			
 			MenuBuilder.AddSubMenu(
-				TAttribute<FText>::Create([this](){ return FText::FromName(SelectedConsoleVariablesAssetName); }),
+				TAttribute<FText>::Create(
+					[this]()
+					{
+						FText OutText = SelectedAdditionalConsoleVariablesAsset.IsValid() ?
+							FText::FromString(SelectedAdditionalConsoleVariablesAsset.GetAssetName()) :
+							LOCTEXT("NoneSelected", "None Selected");
+							
+						return OutText;
+					}),
 				ConsoleVariablesAssetTooltip,
 				FNewMenuDelegate::CreateLambda([this, FoundConsoleVariablesAssets](FMenuBuilder& NewMenuBuilder)
 				{
@@ -1098,7 +1217,7 @@ void FDisplayClusterLaunchEditorModule::AddConsoleVariablesEditorAssetsToToolbar
 							FUIAction(
 								FExecuteAction::CreateRaw(this, &FDisplayClusterLaunchEditorModule::SetSelectedConsoleVariablesAsset, Asset),
 								FCanExecuteAction(),
-								FGetActionCheckState::CreateLambda([this, Asset](){ return SelectedConsoleVariablesAssetName == Asset.AssetName ? ECheckBoxState::Checked : ECheckBoxState::Unchecked; })
+								FGetActionCheckState::CreateLambda([this, Asset](){ return SelectedAdditionalConsoleVariablesAsset.GetAssetName().Equals(Asset.AssetName.ToString()) ? ECheckBoxState::Checked : ECheckBoxState::Unchecked; })
 							),
 							NAME_None,
 							EUserInterfaceActionType::RadioButton
diff --git a/Engine/Plugins/Editor/DisplayClusterLaunch/Source/DisplayClusterLaunchEditor/Public/DisplayClusterLaunchEditorModule.h b/Engine/Plugins/Editor/DisplayClusterLaunch/Source/DisplayClusterLaunchEditor/Public/DisplayClusterLaunchEditorModule.h
--- a/Engine/Plugins/Editor/DisplayClusterLaunch/Source/DisplayClusterLaunchEditor/Public/DisplayClusterLaunchEditorModule.h
+++ b/Engine/Plugins/Editor/DisplayClusterLaunch/Source/DisplayClusterLaunchEditor/Public/DisplayClusterLaunchEditorModule.h
@@ -11,6 +11,8 @@
 #include "UObject/SoftObjectPath.h"
 #include "UObject/WeakObjectPtrTemplates.h"
 
+struct FPlacementCategoryInfo;
+class UDisplayClusterLaunchEditorProjectSettings;
 class ADisplayClusterRootActor;
 class FMenuBuilder;
 class FName;
@@ -83,9 +85,17 @@ private:
 	
 	void RegisterProjectSettings() const;
 
+	static void RegisterPlacementModeItemsIfTheyExist();
+	static const FPlacementCategoryInfo* GetDisplayClusterPlacementCategoryInfo();
+
 	/** Returns a list of selected nodes as FText separated by new lines with the Primary Node marked. */
 	FText GetSelectedNodesListText() const;
 
+	void GetProjectSettingsArguments(
+		const UDisplayClusterLaunchEditorProjectSettings* ProjectSettings, FString& ConcatenatedCommandLineArguments, 
+		FString& ConcatenatedConsoleCommands, FString& ConcatenatedDPCvars, FString& ConcatenatedLogCommands);
+
+
 	TArray<TWeakObjectPtr<ADisplayClusterRootActor>> GetAllDisplayClusterConfigsInWorld();
 	bool DoesCurrentWorldHaveDisplayClusterConfig();
 	void ApplyDisplayClusterConfigOverrides(class UDisplayClusterConfigurationData* ConfigDataCopy);
@@ -116,7 +126,7 @@ private:
 	TArray<FString> SelectedDisplayClusterConfigActorNodes;
 	FString SelectedDisplayClusterConfigActorPrimaryNode;
 	
-	FName SelectedConsoleVariablesAssetName = NAME_None;
+	FSoftObjectPath SelectedAdditionalConsoleVariablesAsset;
 
 	FDisplayClusterLaunchMultiUserServerTrackingData ServerTrackingData;
 	FString CachedConcertSessionName;
diff --git a/Engine/Plugins/Editor/DisplayClusterLaunch/Source/DisplayClusterLaunchEditor/Public/DisplayClusterLaunchEditorProjectSettings.h b/Engine/Plugins/Editor/DisplayClusterLaunch/Source/DisplayClusterLaunchEditor/Public/DisplayClusterLaunchEditorProjectSettings.h
--- a/Engine/Plugins/Editor/DisplayClusterLaunch/Source/DisplayClusterLaunchEditor/Public/DisplayClusterLaunchEditorProjectSettings.h
+++ b/Engine/Plugins/Editor/DisplayClusterLaunch/Source/DisplayClusterLaunchEditor/Public/DisplayClusterLaunchEditorProjectSettings.h
@@ -121,20 +121,25 @@ public:
 	FDirectoryPath ExplicitTraceFileSaveDirectory;
 
 	/**
-	 * A Console Variables Asset to always apply to your launches 
+	 * A Console Variables Asset to always apply to your launches.
+	 * All saved commands and variables will be executed before "Additional" ones specified below.
 	 */
 	UPROPERTY(Config, EditAnywhere, Category="nDisplay Launch Settings|Console", meta = (AllowedClasses="ConsoleVariablesAsset"))
 	FSoftObjectPath ConsoleVariablesPreset;
 
 	/**
 	 * DPCvars
-	 * You can specify additional console variables here to be executed before those of the Console Variable Preset are executed
+	 * You can specify additional console variables here to be executed after those of the Console Variable Preset are executed.
+	 * This is useful for overriding variables in your defined Console Variables Asset.
+	 * Console Variables generally have a name and a value, like "r.ScreenPercentage 50"
 	 */
 	UPROPERTY(Config, EditAnywhere, Category="nDisplay Launch Settings|Console", meta = (Keywords = "DPCvars"))
 	TSet<FString> AdditionalConsoleVariables;
 
 	/**
-	 * You can specify additional console commands here to be executed before those of the Console Variable Preset are executed
+	 * You can specify additional console commands here to be executed after those of the Console Variable Preset are executed.
+	 * This is useful for overriding commands in your defined Console Variables Asset.
+	 * Console Commands generally have a name and no value (i.e. 'stat unit') though some have a value too (i.e. "r.SetNearClipPlane 500")
 	 */
 	UPROPERTY(Config, EditAnywhere, Category="nDisplay Launch Settings|Console")
 	TSet<FString> AdditionalConsoleCommands;
