Fix logic in nDisplay quick launch for existing cases when launching server.

diff --git a/Engine/Plugins/Editor/DisplayClusterLaunch/Source/DisplayClusterLaunchEditor/Private/DisplayClusterLaunchEditorModule.cpp b/Engine/Plugins/Editor/DisplayClusterLaunch/Source/DisplayClusterLaunchEditor/Private/DisplayClusterLaunchEditorModule.cpp
--- a/Engine/Plugins/Editor/DisplayClusterLaunch/Source/DisplayClusterLaunchEditor/Private/DisplayClusterLaunchEditorModule.cpp
+++ b/Engine/Plugins/Editor/DisplayClusterLaunch/Source/DisplayClusterLaunchEditor/Private/DisplayClusterLaunchEditorModule.cpp
@@ -183,35 +183,54 @@ void FDisplayClusterLaunchEditorModule::LaunchConcertServer()
 void FDisplayClusterLaunchEditorModule::FindOrLaunchConcertServer()
 {
 	// Ensure we have the client, otherwise we can't do anything
-	if (const TSharedPtr<IConcertSyncClient> ConcertSyncClient = IConcertSyncClientModule::Get().GetClient(TEXT("MultiUser")))
+	const TSharedPtr<IConcertSyncClient> ConcertSyncClient = IConcertSyncClientModule::Get().GetClient(TEXT("MultiUser"));
+	if (!ConcertSyncClient)
 	{
-		const IConcertClientRef ConcertClient = ConcertSyncClient->GetConcertClient();
-		
-		ConcertClient->OnKnownServersUpdated().RemoveAll(this);
-		
-		// Shutdown existing server no matter what because we need to hook into OnServersAssumedReady
-		IMultiUserClientModule& MultiUserClientModule = IMultiUserClientModule::Get();
-		{
-			if (MultiUserClientModule.IsConcertServerRunning())
-			{
-				// Try to reuse last server
-				ConcertServerRequestStatus = EConcertServerRequestStatus::ReuseExisting;
-				OnServersAssumedReady();
-			}
-			else
+		UE_LOG(LogDisplayClusterLaunchEditor, Error, TEXT("%hs: The ConcertSyncClient could not be found. Please check the output log for errors and try again."), __FUNCTION__);
+		return;
+	}
+
+	const IConcertClientRef ConcertClient = ConcertSyncClient->GetConcertClient();
+	const TSharedPtr<IConcertClientSession> ConcertClientSession = ConcertClient->GetCurrentSession();
+	bool bIsConnectedToSession = ConcertClientSession.IsValid();
+	ConcertClient->OnKnownServersUpdated().RemoveAll(this);
+
+	// Shutdown existing server no matter what because we need to hook into OnServersAssumedReady
+	IMultiUserClientModule& MultiUserClientModule = IMultiUserClientModule::Get();
+
+	if (FPlatformProcess::IsProcRunning(ServerTrackingData.MultiUserServerHandle))
+	{
+		// Try to reuse last server
+		TArray<FConcertServerInfo> Servers = ConcertClient->GetKnownServers();
+		FConcertServerInfo* Server = Algo::FindByPredicate(Servers, [Name = GetConcertServerName()](const FConcertServerInfo& Server)
 			{
-				// Continue when the server list is updated after creation
-				ConcertClient->OnKnownServersUpdated().AddRaw(
-					this, &FDisplayClusterLaunchEditorModule::OnServersAssumedReady
+				return Server.ServerName == Name;
+			});
+		if (Server)
+		{
+			ConcertServerRequestStatus = EConcertServerRequestStatus::ReuseExisting;
+			OnServersAssumedReady();
+		}
+		else
+		{
+			// We are waiting for servers to register with Concert.
+			ConcertClient->OnKnownServersUpdated().AddRaw(
+				this, &FDisplayClusterLaunchEditorModule::OnServersAssumedReady
 				);
-
-				LaunchConcertServer();
-			}
 		}
 	}
+	else if (MultiUserClientModule.IsConcertServerRunning() || bIsConnectedToSession)
+	{
+		ConcertServerRequestStatus = EConcertServerRequestStatus::ReuseExisting;
+		OnServersAssumedReady();
+	}
 	else
 	{
-		UE_LOG(LogDisplayClusterLaunchEditor, Error, TEXT("%hs: The ConcertSyncClient could not be found. Please check the output log for errors and try again."), __FUNCTION__);
+		// We are waiting for servers to register with Concert.
+		ConcertClient->OnKnownServersUpdated().AddRaw(
+			this, &FDisplayClusterLaunchEditorModule::OnServersAssumedReady
+			);
+		LaunchConcertServer();
 	}
 }
 
@@ -238,69 +257,51 @@ void FDisplayClusterLaunchEditorModule::OnServersAssumedReady()
 void FDisplayClusterLaunchEditorModule::FindAppropriateServer()
 {
 	ConcertServerRequestStatus = EConcertServerRequestStatus::None;
-	if (const TSharedPtr<IConcertSyncClient> ConcertSyncClient = IConcertSyncClientModule::Get().GetClient(TEXT("MultiUser")))
+	const TSharedPtr<IConcertSyncClient> ConcertSyncClient = IConcertSyncClientModule::Get().GetClient(TEXT("MultiUser"));
+	if (!ConcertSyncClient)
 	{
-		const IConcertClientRef ConcertClient = ConcertSyncClient->GetConcertClient();
-		if (ConcertClient->GetKnownServers().Num() == 0)
-		{
-			UE_LOG(LogDisplayClusterLaunchEditor, Warning, TEXT("%hs: No servers found. Please launch and connect to one manually."), __FUNCTION__);
-			return;
-		}
-
-		// Try to connect to an existing session even if we launched a new server
-		if (const TSharedPtr<IConcertClientSession> ConcertClientSession = ConcertClient->GetCurrentSession())
-		{
-			const FConcertSessionInfo& SessionInfo = ConcertClientSession->GetSessionInfo();
+		UE_LOG(LogDisplayClusterLaunchEditor, Error, TEXT("%hs: The ConcertSyncClient could not be found. Please check the output log for errors and try again."), __FUNCTION__);
+	}
 
-			// Ensure the reported server is actually running then pull the latest data from it
-			bool bFoundMatch = false;
-			for (const FConcertServerInfo& ServerInfo : ConcertClient->GetKnownServers())
-			{
-				if (ServerInfo.InstanceInfo.InstanceId == SessionInfo.ServerInstanceId)
-				{
-					ServerTrackingData.MultiUserServerInfo = ServerInfo;
-					ServerTrackingData.GeneratedMultiUserServerName = ServerInfo.ServerName;
-					CachedConcertSessionName = ConcertClientSession->GetName();
-					bFoundMatch = true;
-					break;
-				}
-			}
 
-			if (bFoundMatch)
+	const IConcertClientRef ConcertClient = ConcertSyncClient->GetConcertClient();
+	// Try to connect to an existing session even if we launched a new server
+	TArray<FConcertServerInfo> Servers = ConcertClient->GetKnownServers();
+	FConcertServerInfo* Server = nullptr;
+	if (const TSharedPtr<IConcertClientSession> ConcertClientSession = ConcertClient->GetCurrentSession())
+	{
+		const FConcertSessionInfo& SessionInfo = ConcertClientSession->GetSessionInfo();
+		Server = Algo::FindByPredicate(Servers, [&SessionInfo](const FConcertServerInfo& Server)
 			{
-				ConnectToSession();
-			}
-			else
+				return Server.InstanceInfo.InstanceId == SessionInfo.ServerInstanceId;
+			});
+		// The list of known servers must contain the connected session.
+		check(Server);
+		CachedConcertSessionName = ConcertClientSession->GetName();
+	}
+	else
+	{
+		// If no session, we need to try to find a server with a name matching the cached name
+		Server = Algo::FindByPredicate(Servers, [Name = GetConcertServerName()](const FConcertServerInfo& Server)
 			{
-				UE_LOG(LogDisplayClusterLaunchEditor, Error, TEXT("%hs: ConcertClientSession reported a connected server but the server is not in the known servers list."), __FUNCTION__);
-			}
-		}
-		else // If no session, we need to try to find a server with a name matching the cached name
+				return Server.ServerName == Name;
+			});
+		if (!Server && Servers.Num() > 0)
 		{
-			bool bFoundMatch = false;
-			for (const FConcertServerInfo& ServerInfo : ConcertClient->GetKnownServers())
-			{
-				if (ServerInfo.ServerName == GetConcertServerName())
-				{
-					ServerTrackingData.MultiUserServerInfo = ServerInfo;
-					bFoundMatch = true;
-					break;
-				}
-			}
-
-			if (bFoundMatch)
-			{
-				ConnectToSession();
-			}
-			else
-			{
-				UE_LOG(LogDisplayClusterLaunchEditor, Error, TEXT("%hs: Servers exist but a matching server was not found. Try connecting to a server and session manually."), __FUNCTION__);
-			}
+			// Pick a server to use.
+			Server = &Servers[0];
 		}
 	}
+
+	if (Server)
+	{
+		ServerTrackingData.MultiUserServerInfo = *Server;
+		ServerTrackingData.GeneratedMultiUserServerName = Server->ServerName;
+		ConnectToSession();
+	}
 	else
 	{
-		UE_LOG(LogDisplayClusterLaunchEditor, Error, TEXT("%hs: The ConcertSyncClient could not be found. Please check the output log for errors and try again."), __FUNCTION__);
+		UE_LOG(LogDisplayClusterLaunchEditor, Error, TEXT("%hs: Servers exist but a matching server was not found. Try connecting to a server and session manually."), __FUNCTION__);
 	}
 }
 
