nDisplay Launch:

-Update Icons

diff --git a/Engine/Plugins/Editor/DisplayClusterLaunch/Source/DisplayClusterLaunchEditor/Private/DisplayClusterLaunchEditorModule.cpp b/Engine/Plugins/Editor/DisplayClusterLaunch/Source/DisplayClusterLaunchEditor/Private/DisplayClusterLaunchEditorModule.cpp
--- a/Engine/Plugins/Editor/DisplayClusterLaunch/Source/DisplayClusterLaunchEditor/Private/DisplayClusterLaunchEditorModule.cpp
+++ b/Engine/Plugins/Editor/DisplayClusterLaunch/Source/DisplayClusterLaunchEditor/Private/DisplayClusterLaunchEditorModule.cpp
@@ -619,8 +619,8 @@ FText FDisplayClusterLaunchEditorModule::GetToolbarButtonTooltipText()
 FSlateIcon FDisplayClusterLaunchEditorModule::GetToolbarButtonIcon()
 {
 	RemoveTerminatedNodeProcesses();
-	return FSlateIcon(FAppStyle::Get().GetStyleSetName(),
-		ActiveDisplayClusterProcesses.Num() > 0 ? "Icons.Toolbar.Stop" : "Icons.Toolbar.Play");
+	return FSlateIcon(FDisplayClusterLaunchEditorStyle::Get().GetStyleSetName(),
+		ActiveDisplayClusterProcesses.Num() > 0 ? "Icons.DisplayClusterLaunchStop" : "Icons.DisplayClusterLaunchPlay");
 }
 
 void FDisplayClusterLaunchEditorModule::OnClickToolbarButton()
diff --git a/Engine/Plugins/Editor/DisplayClusterLaunch/Source/DisplayClusterLaunchEditor/Private/DisplayClusterLaunchEditorStyle.cpp b/Engine/Plugins/Editor/DisplayClusterLaunch/Source/DisplayClusterLaunchEditor/Private/DisplayClusterLaunchEditorStyle.cpp
--- a/Engine/Plugins/Editor/DisplayClusterLaunch/Source/DisplayClusterLaunchEditor/Private/DisplayClusterLaunchEditorStyle.cpp
+++ b/Engine/Plugins/Editor/DisplayClusterLaunch/Source/DisplayClusterLaunchEditor/Private/DisplayClusterLaunchEditorStyle.cpp
@@ -76,10 +76,8 @@ FString FDisplayClusterLaunchEditorStyle::GetExternalPluginContent(const FString
 TSharedRef< FSlateStyleSet > FDisplayClusterLaunchEditorStyle::Create()
 {
 	TSharedRef<FSlateStyleSet> Style = MakeShared<FSlateStyleSet>("DisplayClusterLaunchEditor");
-
-	TSharedPtr<IPlugin> Plugin = IPluginManager::Get().FindPlugin(TEXT("DisplayClusterLaunch"));
-	check(Plugin.IsValid());
-	if (Plugin.IsValid())
+	
+	if (TSharedPtr<IPlugin> Plugin = IPluginManager::Get().FindPlugin(TEXT("DisplayClusterLaunch")))
 	{
 		Style->SetContentRoot(FPaths::Combine(Plugin->GetBaseDir(), TEXT("Resources")));
 	}
@@ -87,13 +85,11 @@ TSharedRef< FSlateStyleSet > FDisplayClusterLaunchEditorStyle::Create()
 	Style->SetCoreContentRoot(FPaths::EngineContentDir() / TEXT("Slate"));
 
 	// Icons
-	// todo: update
-	Style->Set("ConsoleVariables.ToolbarButton", new IMAGE_BRUSH_SVG("Icons/ConsoleVariables", Icon40x40));
-	Style->Set("ConsoleVariables.ToolbarButton.Small", new IMAGE_BRUSH_SVG("Icons/ConsoleVariables", Icon20x20));
+	Style->Set("Icons.DisplayClusterLaunchPlay", new IMAGE_BRUSH_SVG("Icons/nDisplayQuickLaunchPlay_20", Icon16x16));
+	Style->Set("Icons.DisplayClusterLaunchStop", new IMAGE_BRUSH_SVG("Icons/nDisplayQuickLaunchX_20", Icon16x16));
 
 	// External plugin icons
 	
-	Style->Set("Icons.ConsoleVariablesEditor", new IMAGE_PLUGIN_BRUSH_SVG("ConsoleVariables","Resources/Icons/ConsoleVariables", Icon16x16));
 	Style->Set("Icons.MultiUser", new IMAGE_PLUGIN_BRUSH_SVG("ConcertSharedSlate","Content/Icons/icon_MultiUser", Icon16x16));
 	Style->Set("Icons.DisplayCluster", new IMAGE_PLUGIN_BRUSH_SVG("nDisplay","Content/Icons/RootActor/nDisplay", Icon16x16));
 	Style->Set("Icons.DisplayClusterNode", new IMAGE_PLUGIN_BRUSH_SVG("nDisplay","Content/Icons/Cluster/ClusterNode", Icon16x16));
