nDisplay Launch:

-Update Icons

diff --git a/Engine/Plugins/Editor/DisplayClusterLaunch/Source/DisplayClusterLaunchEditor/Private/DisplayClusterLaunchEditorModule.cpp b/Engine/Plugins/Editor/DisplayClusterLaunch/Source/DisplayClusterLaunchEditor/Private/DisplayClusterLaunchEditorModule.cpp
--- a/Engine/Plugins/Editor/DisplayClusterLaunch/Source/DisplayClusterLaunchEditor/Private/DisplayClusterLaunchEditorModule.cpp
+++ b/Engine/Plugins/Editor/DisplayClusterLaunch/Source/DisplayClusterLaunchEditor/Private/DisplayClusterLaunchEditorModule.cpp
@@ -976,7 +976,7 @@ TSharedRef<SWidget>  FDisplayClusterLaunchEditorModule::CreateToolbarMenuEntries
 		MenuBuilder.AddMenuEntry(
 			LOCTEXT("DisplayClusterLaunchLastNode", "Launch Last Node Configuration"),
 			LOCTEXT("DisplayClusterLaunchLastNodeTooltip", "Launch the last node configuration."),
-			FSlateIcon(FAppStyle::Get().GetStyleSetName(), "Icons.Toolbar.Play"),
+			FSlateIcon(FDisplayClusterLaunchEditorStyle::Get().GetStyleSetName(), "Icons.DisplayClusterLaunchPlay"),
 			FUIAction(
 				FExecuteAction::CreateRaw(this, &FDisplayClusterLaunchEditorModule::TryLaunchDisplayClusterProcess),
 				FCanExecuteAction::CreateRaw(this, &FDisplayClusterLaunchEditorModule::DoesCurrentWorldHaveDisplayClusterConfig)
