Remove unused duplicate functions causing linker issues in monolithic builds

diff --git a/Engine/Plugins/Editor/ConsoleVariablesEditor/Source/ConsoleVariablesEditor/Private/ConsoleVariablesEditorStyle.cpp b/Engine/Plugins/Editor/ConsoleVariablesEditor/Source/ConsoleVariablesEditor/Private/ConsoleVariablesEditorStyle.cpp
--- a/Engine/Plugins/Editor/ConsoleVariablesEditor/Source/ConsoleVariablesEditor/Private/ConsoleVariablesEditorStyle.cpp
+++ b/Engine/Plugins/Editor/ConsoleVariablesEditor/Source/ConsoleVariablesEditor/Private/ConsoleVariablesEditorStyle.cpp
@@ -44,12 +44,6 @@ const FName& FConsoleVariablesEditorStyle::GetStyleSetName() const
 	return ConsoleVariablesStyleSetName;
 }
 
-FString GetConcertContent(const FString& RelativePath, const ANSICHAR* Extension)
-{
-	static FString ContentDir = IPluginManager::Get().FindPlugin(TEXT("ConcertSyncClient"))->GetContentDir();
-	return (ContentDir / RelativePath) + Extension;
-}
-
 const FSlateBrush* FConsoleVariablesEditorStyle::GetBrush(const FName PropertyName, const ANSICHAR* Specifier, const ISlateStyle* RequestingStyle) const
 {
 	return StyleInstance->GetBrush(PropertyName, Specifier);
diff --git a/Engine/Plugins/Editor/DisplayClusterLaunch/Source/DisplayClusterLaunchEditor/Private/DisplayClusterLaunchEditorStyle.cpp b/Engine/Plugins/Editor/DisplayClusterLaunch/Source/DisplayClusterLaunchEditor/Private/DisplayClusterLaunchEditorStyle.cpp
--- a/Engine/Plugins/Editor/DisplayClusterLaunch/Source/DisplayClusterLaunchEditor/Private/DisplayClusterLaunchEditorStyle.cpp
+++ b/Engine/Plugins/Editor/DisplayClusterLaunch/Source/DisplayClusterLaunchEditor/Private/DisplayClusterLaunchEditorStyle.cpp
@@ -43,12 +43,6 @@ const FName& FDisplayClusterLaunchEditorStyle::GetStyleSetName() const
 	return DisplayClusterLaunchStyleSetName;
 }
 
-FString GetConcertContent(const FString& RelativePath, const ANSICHAR* Extension)
-{
-	static FString ContentDir = IPluginManager::Get().FindPlugin(TEXT("ConcertSyncClient"))->GetContentDir();
-	return (ContentDir / RelativePath) + Extension;
-}
-
 const FSlateBrush* FDisplayClusterLaunchEditorStyle::GetBrush(const FName PropertyName, const ANSICHAR* Specifier, const ISlateStyle* RequestingStyle) const
 {
 	return StyleInstance->GetBrush(PropertyName, Specifier);
diff --git a/Engine/Plugins/Editor/ObjectMixer/ObjectMixer/Source/ObjectMixer/Private/ObjectMixerEditorStyle.cpp b/Engine/Plugins/Editor/ObjectMixer/ObjectMixer/Source/ObjectMixer/Private/ObjectMixerEditorStyle.cpp
--- a/Engine/Plugins/Editor/ObjectMixer/ObjectMixer/Source/ObjectMixer/Private/ObjectMixerEditorStyle.cpp
+++ b/Engine/Plugins/Editor/ObjectMixer/ObjectMixer/Source/ObjectMixer/Private/ObjectMixerEditorStyle.cpp
@@ -44,12 +44,6 @@ const FName& FObjectMixerEditorStyle::GetStyleSetName() const
 	return StyleName;
 }
 
-FString GetConcertContent(const FString& RelativePath, const ANSICHAR* Extension)
-{
-	static FString ContentDir = IPluginManager::Get().FindPlugin(TEXT("ConcertSyncClient"))->GetContentDir();
-	return (ContentDir / RelativePath) + Extension;
-}
-
 const FSlateBrush* FObjectMixerEditorStyle::GetBrush(const FName PropertyName, const ANSICHAR* Specifier, const ISlateStyle* RequestingStyle) const
 {
 	return StyleInstance->GetBrush(PropertyName, Specifier);
