nDisplay Launch:

-Fix: New Console Variable presets created are not available in nDisplay Launch menu
-Fix: Null Console Variable presets are present in nDisplay Launch menu
-Fix: Extraneous warning when switching configs in the menu

diff --git a/Engine/Plugins/Editor/DisplayClusterLaunch/Source/DisplayClusterLaunchEditor/Private/DisplayClusterLaunchEditorModule.cpp b/Engine/Plugins/Editor/DisplayClusterLaunch/Source/DisplayClusterLaunchEditor/Private/DisplayClusterLaunchEditorModule.cpp
--- a/Engine/Plugins/Editor/DisplayClusterLaunch/Source/DisplayClusterLaunchEditor/Private/DisplayClusterLaunchEditorModule.cpp
+++ b/Engine/Plugins/Editor/DisplayClusterLaunch/Source/DisplayClusterLaunchEditor/Private/DisplayClusterLaunchEditorModule.cpp
@@ -276,7 +276,11 @@ void FDisplayClusterLaunchEditorModule::FindAppropriateServer()
 				return Server.InstanceInfo.InstanceId == SessionInfo.ServerInstanceId;
 			});
 		// The list of known servers must contain the connected session.
-		check(Server);
+		if (!Server)
+		{
+			UE_LOG(LogDisplayClusterLaunchEditor, Error, TEXT("%hs: The list of known servers does not contain the connected session. Please check the output log for errors and try again."), __FUNCTION__);
+			return;
+		}
 		CachedConcertSessionName = ConcertClientSession->GetName();
 	}
 	else
@@ -327,7 +331,11 @@ void FDisplayClusterLaunchEditorModule::ConnectToSession()
 		AutoConnectConfig->DefaultSessionName = GetConcertSessionName();
 		
 		ConcertClient->Configure(AutoConnectConfig);
-		check(ConcertClient->IsConfigured());
+		if (!ConcertClient->IsConfigured())
+		{
+			UE_LOG(LogDisplayClusterLaunchEditor, Error, TEXT("%hs: The ConcertSyncClient is not configured. Please check the output log for errors and try again."), __FUNCTION__);
+			return;
+		}
 		
 		// Initiate the auto connect to the named server and session.
 		if (ConcertClient->CanAutoConnect())
@@ -956,9 +964,9 @@ void FDisplayClusterLaunchEditorModule::SelectFirstNode(ADisplayClusterRootActor
 	if (SelectedDisplayClusterConfigActorNodes.Num() == 0)
 	{
 		const FString& NodeName = NodeNames[0];
-		UE_LOG(LogDisplayClusterLaunchEditor, Warning, TEXT("%hs: Selected nDisplay nodes were not found on the selected DisplayClusterRootActor. We will select the first valid node."), __FUNCTION__);
+		UE_LOG(LogDisplayClusterLaunchEditor, Verbose, TEXT("%hs: Selected nDisplay nodes were not found on the selected DisplayClusterRootActor. We will select the first valid node."), __FUNCTION__);
 		SelectedDisplayClusterConfigActorNodes.Add(NodeName);
-		UE_LOG(LogDisplayClusterLaunchEditor, Log, TEXT("%hs: Adding first valid node named '%s' to selected nodes."), __FUNCTION__, *NodeName);
+		UE_LOG(LogDisplayClusterLaunchEditor, Verbose, TEXT("%hs: Adding first valid node named '%s' to selected nodes."), __FUNCTION__, *NodeName);
 	}
 
 	// Set the primary node to be the first in the list
@@ -1187,7 +1195,7 @@ void FDisplayClusterLaunchEditorModule::AddConsoleVariablesEditorAssetsToToolbar
 	IAssetRegistry* AssetRegistry, FMenuBuilder& MenuBuilder)
 {
 	TArray<FAssetData> FoundConsoleVariablesAssets;
-	AssetRegistry->GetAssetsByClass(FTopLevelAssetPath(TEXT("/Script/ConsoleVariablesEditor"), TEXT("ConsoleVariablesAsset")), FoundConsoleVariablesAssets, true);
+	AssetRegistry->GetAssetsByClass(FTopLevelAssetPath(TEXT("/Script/ConsoleVariablesEditorRuntime"), TEXT("ConsoleVariablesAsset")), FoundConsoleVariablesAssets, true);
 	
 	if (FoundConsoleVariablesAssets.Num())
 	{
