nDisplay Launch: Use blueprint/asset as config source rather than placed actor instance.

Other changes:
 - Write nDisplay config JSON to a random temporary path instead of "Saved/Temp.ndisplay".
 - Pass an absolute path on the command line instead of relative.
 - No longer include obsolete r.Shadow.Virtual.Cache=0 override.

diff --git a/Engine/Plugins/Editor/DisplayClusterLaunch/Source/DisplayClusterLaunchEditor/Private/DisplayClusterLaunchEditorModule.cpp b/Engine/Plugins/Editor/DisplayClusterLaunch/Source/DisplayClusterLaunchEditor/Private/DisplayClusterLaunchEditorModule.cpp
--- a/Engine/Plugins/Editor/DisplayClusterLaunch/Source/DisplayClusterLaunchEditor/Private/DisplayClusterLaunchEditorModule.cpp
+++ b/Engine/Plugins/Editor/DisplayClusterLaunch/Source/DisplayClusterLaunchEditor/Private/DisplayClusterLaunchEditorModule.cpp
@@ -5,6 +5,7 @@
 #include "DisplayClusterLaunchEditorProjectSettings.h"
 #include "DisplayClusterLaunchEditorStyle.h"
 
+#include "Blueprints/DisplayClusterBlueprint.h"
 #include "DisplayClusterRootActor.h"
 #include "IDisplayClusterConfiguration.h"
 #include "DisplayClusterConfigurationTypes.h"
@@ -158,7 +159,7 @@ FString AppendRandomNumbersToString(const FString InString, uint8 NumberToAppend
 FString GetConcertArguments(const FString& ServerName, const FString& SessionName)
 {
 	const UConcertClientConfig* ConcertClientConfig = GetDefault<UConcertClientConfig>();
-	ensureAlwaysMsgf (ConcertClientConfig, TEXT("%hs: Unable to launch nDisplay because there is no UConcertClientConfig object."));
+	ensureAlwaysMsgf(ConcertClientConfig, TEXT("%hs: Unable to launch nDisplay because there is no UConcertClientConfig object."), __FUNCTION__);
 	FString ReturnValue =
 		FString::Printf(TEXT("-CONCERTISHEADLESS -CONCERTRETRYAUTOCONNECTONERROR -CONCERTAUTOCONNECT -CONCERTSERVER=\"%s\" -CONCERTSESSION=\"%s\""),
 			 *ServerName, *SessionName);
@@ -353,7 +354,7 @@ void FDisplayClusterLaunchEditorModule::ConnectToSession()
 
 void FDisplayClusterLaunchEditorModule::TryLaunchDisplayClusterProcess()
 {
-	if (!ensureAlwaysMsgf(GetDefault<UDisplayClusterLaunchEditorProjectSettings>(), TEXT("%hs: Unable to launch nDisplay because there is no UDisplayClusterLaunchEditorProjectSettings object.")))
+	if (!ensureAlwaysMsgf(GetDefault<UDisplayClusterLaunchEditorProjectSettings>(), TEXT("%hs: Unable to launch nDisplay because there is no UDisplayClusterLaunchEditorProjectSettings object."), __FUNCTION__))
 	{
 		return;
 	}
@@ -439,10 +440,13 @@ void FDisplayClusterLaunchEditorModule::LaunchDisplayClusterProcess()
 	// If it's valid we need to check the selected nodes against the current config. If they don't exist, we need to get the first one.
 	if (const ADisplayClusterRootActor* ConfigActor = Cast<ADisplayClusterRootActor>(SelectedDisplayClusterConfigActor.ResolveObject()))
 	{
-		// Duplicate existing config data so we can make non-destructive edits
-		ConfigDataToUse = DuplicateObject(ConfigActor->GetConfigData(), GetTransientPackage());
+		// Duplicate existing config data so we can make non-destructive edits.
+		UDisplayClusterBlueprint* ConfigBlueprint = CastChecked<UDisplayClusterBlueprint>(ConfigActor->GetClass()->ClassGeneratedBy);
+		ConfigDataToUse = DuplicateObject(ConfigBlueprint->GetOrLoadConfig(), GetTransientPackage());
+
 		ApplyDisplayClusterConfigOverrides(ConfigDataToUse);
-		const FString FilePath = FPaths::Combine(FPaths::ProjectSavedDir(), "Temp.ndisplay");
+
+		const FString FilePath = FPaths::CreateTempFilename(FPlatformProcess::UserTempDir(), TEXT(""), TEXT(".ndisplay"));
 		if (!ensureAlways(IDisplayClusterConfiguration::Get().SaveConfig(ConfigDataToUse, FilePath)))
 		{
 			UE_LOG(LogDisplayClusterLaunchEditor, Error, TEXT("%hs: Unable to launch nDisplay because the selected nDisplay Configuration could not be saved to a .ndisplay file. See the log for more information."), __FUNCTION__);
diff --git a/Engine/Plugins/Editor/DisplayClusterLaunch/Source/DisplayClusterLaunchEditor/Public/DisplayClusterLaunchEditorProjectSettings.h b/Engine/Plugins/Editor/DisplayClusterLaunch/Source/DisplayClusterLaunchEditor/Public/DisplayClusterLaunchEditorProjectSettings.h
--- a/Engine/Plugins/Editor/DisplayClusterLaunch/Source/DisplayClusterLaunchEditor/Public/DisplayClusterLaunchEditorProjectSettings.h
+++ b/Engine/Plugins/Editor/DisplayClusterLaunch/Source/DisplayClusterLaunchEditor/Public/DisplayClusterLaunchEditorProjectSettings.h
@@ -74,7 +74,7 @@ public:
 
 		AdditionalConsoleCommands = { "DisableAllScreenMessages" };
 
-		AdditionalConsoleVariables = { "p.Chaos.Solver.Deterministic=1", "r.Shadow.Virtual.Cache=0" };
+		AdditionalConsoleVariables = { "p.Chaos.Solver.Deterministic=1" };
 
 		Logging = { {"LogDisplayClusterRenderSync", EDisplayClusterLaunchLogVerbosity::Log } };
 	}
