nDisplayLauncher: Now uses correct graphics adapter and headless mode.

diff --git a/Engine/Plugins/Editor/DisplayClusterLaunch/Source/DisplayClusterLaunchEditor/Private/DisplayClusterLaunchEditorModule.cpp b/Engine/Plugins/Editor/DisplayClusterLaunch/Source/DisplayClusterLaunchEditor/Private/DisplayClusterLaunchEditorModule.cpp
--- a/Engine/Plugins/Editor/DisplayClusterLaunch/Source/DisplayClusterLaunchEditor/Private/DisplayClusterLaunchEditorModule.cpp
+++ b/Engine/Plugins/Editor/DisplayClusterLaunch/Source/DisplayClusterLaunchEditor/Private/DisplayClusterLaunchEditorModule.cpp
@@ -472,8 +472,10 @@ void FDisplayClusterLaunchEditorModule::LaunchDisplayClusterProcess()
 	{
 		FString ConcatenatedCommandLineArguments;
 		FString ConcatenatedConsoleCommands;
-		FString ConcatenatedDPCvars;
 		FString ConcatenatedLogCommands;
+
+		TSet<FString> DPCvars;
+
 		// Fullscreen/Windowed
 		if (TObjectPtr<UDisplayClusterConfigurationClusterNode>* NodePtrPtr = ConfigDataToUse->Cluster->Nodes.Find(Node))
 		{
@@ -492,6 +494,18 @@ void FDisplayClusterLaunchEditorModule::LaunchDisplayClusterProcess()
 								NodePtr->WindowRect.W, NodePtr->WindowRect.H
 						);
 				}
+
+				// Headless
+				if (NodePtr->bRenderHeadless)
+				{
+					ConcatenatedCommandLineArguments += "-RenderOffscreen ";
+				}
+
+				// Graphics adapter
+				if (NodePtr->GraphicsAdapter >= 0)
+				{
+					DPCvars.Add(FString::Printf(TEXT("r.GraphicsAdapter=%d"), NodePtr->GraphicsAdapter));
+				}
 			}
 		}
 		// Unreal Insights support
@@ -513,7 +527,7 @@ void FDisplayClusterLaunchEditorModule::LaunchDisplayClusterProcess()
 			}
 		}
 	
-		GetProjectSettingsArguments(ProjectSettings, ConcatenatedCommandLineArguments, ConcatenatedConsoleCommands, ConcatenatedDPCvars,
+		GetProjectSettingsArguments(ProjectSettings, ConcatenatedCommandLineArguments, ConcatenatedConsoleCommands, DPCvars,
 					 ConcatenatedLogCommands);
 		
 		AddUdpMessagingArguments(ConcatenatedCommandLineArguments);
@@ -531,8 +545,7 @@ void FDisplayClusterLaunchEditorModule::LaunchDisplayClusterProcess()
 			FString::Printf(
 				TEXT("\"%s\" -game \"%s\" Log=%s %s -ExecCmds=\"%s\" -DPCVars=\"%s\" -LogCmds=\"%s\""),
 				*Project, *Map, *LogFileName, *ConcatenatedCommandLineArguments,
-				*ConcatenatedConsoleCommands, *ConcatenatedDPCvars, *ConcatenatedLogCommands
-	
+				*ConcatenatedConsoleCommands, *FString::Join(DPCvars, TEXT(",")), *ConcatenatedLogCommands
 		);
 		UE_LOG(LogDisplayClusterLaunchEditor, Log, TEXT("Full Command: %s %s"), *EditorBinary, *Params);
 		void* WritePipe = nullptr;
@@ -762,7 +775,7 @@ FText FDisplayClusterLaunchEditorModule::GetSelectedNodesListText() const
 
 void FDisplayClusterLaunchEditorModule::GetProjectSettingsArguments(
 	const UDisplayClusterLaunchEditorProjectSettings* ProjectSettings, FString& ConcatenatedCommandLineArguments,
-	FString& ConcatenatedConsoleCommands, FString& ConcatenatedDPCvars, FString& ConcatenatedLogCommands)
+	FString& ConcatenatedConsoleCommands, TSet<FString>& DPCvars, FString& ConcatenatedLogCommands)
 {
 	{
 		for (const FString& CommandLineArgument : ProjectSettings->CommandLineArguments)
@@ -803,7 +816,7 @@ void FDisplayClusterLaunchEditorModule::GetProjectSettingsArguments(
 	ConcatenatedConsoleCommands += FString::Join(CollectedCommands, TEXT(","));
 
 	// Additional Console Variables
-	ConcatenatedDPCvars += FString::Join(ProjectSettings->AdditionalConsoleVariables, TEXT(","));
+	DPCvars.Append(ProjectSettings->AdditionalConsoleVariables);
 
 	// Logging
 	{
diff --git a/Engine/Plugins/Editor/DisplayClusterLaunch/Source/DisplayClusterLaunchEditor/Public/DisplayClusterLaunchEditorModule.h b/Engine/Plugins/Editor/DisplayClusterLaunch/Source/DisplayClusterLaunchEditor/Public/DisplayClusterLaunchEditorModule.h
--- a/Engine/Plugins/Editor/DisplayClusterLaunch/Source/DisplayClusterLaunchEditor/Public/DisplayClusterLaunchEditorModule.h
+++ b/Engine/Plugins/Editor/DisplayClusterLaunch/Source/DisplayClusterLaunchEditor/Public/DisplayClusterLaunchEditorModule.h
@@ -93,7 +93,7 @@ private:
 
 	void GetProjectSettingsArguments(
 		const UDisplayClusterLaunchEditorProjectSettings* ProjectSettings, FString& ConcatenatedCommandLineArguments, 
-		FString& ConcatenatedConsoleCommands, FString& ConcatenatedDPCvars, FString& ConcatenatedLogCommands);
+		FString& ConcatenatedConsoleCommands, TSet<FString>& DPCvars, FString& ConcatenatedLogCommands);
 
 
 	TArray<TWeakObjectPtr<ADisplayClusterRootActor>> GetAllDisplayClusterConfigsInWorld();
