nDisplay Launch:

-Address: Editor UI infrequently becomes unresponsive following launch

diff --git a/Engine/Plugins/Editor/DisplayClusterLaunch/Source/DisplayClusterLaunchEditor/Public/DisplayClusterLaunchEditorProjectSettings.h b/Engine/Plugins/Editor/DisplayClusterLaunch/Source/DisplayClusterLaunchEditor/Public/DisplayClusterLaunchEditorProjectSettings.h
--- a/Engine/Plugins/Editor/DisplayClusterLaunch/Source/DisplayClusterLaunchEditor/Public/DisplayClusterLaunchEditorProjectSettings.h
+++ b/Engine/Plugins/Editor/DisplayClusterLaunch/Source/DisplayClusterLaunchEditor/Public/DisplayClusterLaunchEditorProjectSettings.h
@@ -124,7 +124,7 @@ public:
 	 * A Console Variables Asset to always apply to your launches.
 	 * All saved commands and variables will be executed before "Additional" ones specified below.
 	 */
-	UPROPERTY(Config, EditAnywhere, Category="nDisplay Launch Settings|Console", meta = (AllowedClasses="ConsoleVariablesAsset"))
+	UPROPERTY(Config, EditAnywhere, Category="nDisplay Launch Settings|Console", meta = (AllowedClasses="/Script/ConsoleVariablesEditorRuntime.ConsoleVariablesAsset"))
 	FSoftObjectPath ConsoleVariablesPreset;
 
 	/**
