Make sure we start the discovery process even if the Multi-user browser is not open.

diff --git a/Engine/Plugins/Editor/DisplayClusterLaunch/Source/DisplayClusterLaunchEditor/Private/DisplayClusterLaunchEditorModule.cpp b/Engine/Plugins/Editor/DisplayClusterLaunch/Source/DisplayClusterLaunchEditor/Private/DisplayClusterLaunchEditorModule.cpp
--- a/Engine/Plugins/Editor/DisplayClusterLaunch/Source/DisplayClusterLaunchEditor/Private/DisplayClusterLaunchEditorModule.cpp
+++ b/Engine/Plugins/Editor/DisplayClusterLaunch/Source/DisplayClusterLaunchEditor/Private/DisplayClusterLaunchEditorModule.cpp
@@ -80,6 +80,14 @@ void FDisplayClusterLaunchEditorModule::StartupModule()
 {
 	FDisplayClusterLaunchEditorStyle::Initialize();
 	FCoreDelegates::OnFEngineLoopInitComplete.AddRaw(this, &FDisplayClusterLaunchEditorModule::OnFEngineLoopInitComplete);
+	if (IConcertSyncClientModule* ConcertSyncClientModule = (IConcertSyncClientModule*)FModuleManager::Get().GetModule("ConcertSyncClient"))
+	{
+		if (const TSharedPtr<IConcertSyncClient> ConcertSyncClient = ConcertSyncClientModule->GetClient(TEXT("MultiUser")))
+		{
+			const IConcertClientRef ConcertClient = ConcertSyncClient->GetConcertClient();
+			ConcertClient->StartDiscovery();
+		}
+	}
 }
 
 void FDisplayClusterLaunchEditorModule::ShutdownModule()
@@ -94,13 +102,18 @@ void FDisplayClusterLaunchEditorModule::ShutdownModule()
 		SettingsModule.UnregisterSettings("Project", "Plugins", "nDisplay Launch");
 	}
 
-	// Remove Concert delegates
 	if (IConcertSyncClientModule* ConcertSyncClientModule = (IConcertSyncClientModule*)FModuleManager::Get().GetModule("ConcertSyncClient"))
 	{
 		if (const TSharedPtr<IConcertSyncClient> ConcertSyncClient = ConcertSyncClientModule->GetClient(TEXT("MultiUser")))
 		{
 			const IConcertClientRef ConcertClient = ConcertSyncClient->GetConcertClient();
-		
+
+			// Concert may close all existing discovery requests so we have to check to see if we are still have discovery enabled before
+			// attempting to stop.
+			if (ConcertClient->IsDiscoveryEnabled())
+			{
+				ConcertClient->StopDiscovery();
+			}
 			ConcertClient->OnKnownServersUpdated().RemoveAll(this);
 		}
 	}
