Move some code around in DisplayCluster launch.

#rnx

diff --git a/Engine/Plugins/Editor/DisplayClusterLaunch/Config/BaseDisplayClusterLaunch.ini b/Engine/Plugins/Editor/DisplayClusterLaunch/Config/BaseDisplayClusterLaunch.ini
new file mode 100644
--- /dev/null
+++ b/Engine/Plugins/Editor/DisplayClusterLaunch/Config/BaseDisplayClusterLaunch.ini
@@ -0,0 +1,5 @@
+[/Script/DisplayClusterLaunchEditor.DisplayClusterLaunchEditorProjectSettings]
+AdditionalConsoleVariables=("p.Chaos.Solver.Deterministic=1","LevelInstance.ForceEditorWorldMode=1")
+AdditionalConsoleCommands=("DisableAllScreenMessages")
+CommandLineArguments=("messaging","dc_cluster","nosplash","fixedseed","NoVerifyGC","noxrstereo","xrtrackingonly","RemoteControlIsHeadless","dx12","dc_dev_mono","unattended","handleensurepercent=0","ini:Engine:[/Script/Engine.Engine]:GameEngine=/Script/DisplayCluster.DisplayClusterGameEngine,[/Script/Engine.Engine]:GameViewportClientClassName=/Script/DisplayCluster.DisplayClusterViewportClient,[/Script/Engine.UserInterfaceSettings]:bAllowHighDPIInGameMode=True","ini:Game:[/Script/EngineSettings.GeneralProjectSettings]:bUseBorderlessWindow=True")
+Logging=(Category="LogDisplayClusterRenderSync",VerbosityLevel=Log)
diff --git a/Engine/Plugins/Editor/DisplayClusterLaunch/Source/DisplayClusterLaunchEditor/Public/DisplayClusterLaunchEditorProjectSettings.h b/Engine/Plugins/Editor/DisplayClusterLaunch/Source/DisplayClusterLaunchEditor/Public/DisplayClusterLaunchEditorProjectSettings.h
--- a/Engine/Plugins/Editor/DisplayClusterLaunch/Source/DisplayClusterLaunchEditor/Public/DisplayClusterLaunchEditorProjectSettings.h
+++ b/Engine/Plugins/Editor/DisplayClusterLaunch/Source/DisplayClusterLaunchEditor/Public/DisplayClusterLaunchEditorProjectSettings.h
@@ -62,23 +62,6 @@ class DISPLAYCLUSTERLAUNCHEDITOR_API UDisplayClusterLaunchEditorProjectSettings
 	
 public:
 
-	UDisplayClusterLaunchEditorProjectSettings(const FObjectInitializer& ObjectInitializer)
-	{
-		CommandLineArguments =
-		{
-			"messaging", "dc_cluster", "nosplash", "fixedseed", "NoVerifyGC", "noxrstereo", "xrtrackingonly", "RemoteControlIsHeadless",
-			"dx12", "dc_dev_mono", "unattended", "handleensurepercent=0", 
-			"ini:Engine:[/Script/Engine.Engine]:GameEngine=/Script/DisplayCluster.DisplayClusterGameEngine,[/Script/Engine.Engine]:GameViewportClientClassName=/Script/DisplayCluster.DisplayClusterViewportClient,[/Script/Engine.UserInterfaceSettings]:bAllowHighDPIInGameMode=True",
-			"ini:Game:[/Script/EngineSettings.GeneralProjectSettings]:bUseBorderlessWindow=True"
-		};
-
-		AdditionalConsoleCommands = { "DisableAllScreenMessages" };
-
-		AdditionalConsoleVariables = { "p.Chaos.Solver.Deterministic=1" };
-
-		Logging = { {"LogDisplayClusterRenderSync", EDisplayClusterLaunchLogVerbosity::Log } };
-	}
-
 	/**
 	 * If true, the editor will be closed on session launch to optimize session performance.
 	 */
