nDisplay: added missing DisplayClusterInput parameter to the nDisplay QuickLaunch

diff --git a/Engine/Plugins/Editor/DisplayClusterLaunch/Source/DisplayClusterLaunchEditor/Private/DisplayClusterLaunchEditorModule.cpp b/Engine/Plugins/Editor/DisplayClusterLaunch/Source/DisplayClusterLaunchEditor/Private/DisplayClusterLaunchEditorModule.cpp
--- a/Engine/Plugins/Editor/DisplayClusterLaunch/Source/DisplayClusterLaunchEditor/Private/DisplayClusterLaunchEditorModule.cpp
+++ b/Engine/Plugins/Editor/DisplayClusterLaunch/Source/DisplayClusterLaunchEditor/Private/DisplayClusterLaunchEditorModule.cpp
@@ -176,6 +176,12 @@ FString GetConcertArguments(const FString& ServerName, const FString& SessionNam
 	return ReturnValue;
 }
 
+const FString GetInputArguments()
+{
+	const FString ReturnValue = FString::Printf(TEXT("-ini:Input:[/Script/Engine.InputSettings]:DefaultPlayerInputClass=/Script/DisplayCluster.DisplayClusterPlayerInput"));
+	return ReturnValue;
+}
+
 void FDisplayClusterLaunchEditorModule::LaunchConcertServer()
 {
 	IMultiUserClientModule& MultiUserClientModule = IMultiUserClientModule::Get();
@@ -547,6 +553,10 @@ void FDisplayClusterLaunchEditorModule::LaunchDisplayClusterProcess()
 		{
 			ConcatenatedCommandLineArguments += " " + ConcertArguments;
 		}
+
+		// Input arguments
+		ConcatenatedCommandLineArguments += " " + GetInputArguments();
+
 		// Log file
 		const FString LogFileName = (ProjectSettings->LogFileName.IsEmpty() ? Node : ProjectSettings->LogFileName) + ".log";
 		const FString Params =
