Fix logic around starting up Multi-user discovery from nDisplay quick launch tool

diff --git a/Engine/Plugins/Editor/DisplayClusterLaunch/Source/DisplayClusterLaunchEditor/Private/DisplayClusterLaunchEditorModule.cpp b/Engine/Plugins/Editor/DisplayClusterLaunch/Source/DisplayClusterLaunchEditor/Private/DisplayClusterLaunchEditorModule.cpp
--- a/Engine/Plugins/Editor/DisplayClusterLaunch/Source/DisplayClusterLaunchEditor/Private/DisplayClusterLaunchEditorModule.cpp
+++ b/Engine/Plugins/Editor/DisplayClusterLaunch/Source/DisplayClusterLaunchEditor/Private/DisplayClusterLaunchEditorModule.cpp
@@ -75,14 +75,6 @@ void FDisplayClusterLaunchEditorModule::StartupModule()
 {
 	FDisplayClusterLaunchEditorStyle::Initialize();
 	FCoreDelegates::OnFEngineLoopInitComplete.AddRaw(this, &FDisplayClusterLaunchEditorModule::OnFEngineLoopInitComplete);
-	if (IConcertSyncClientModule* ConcertSyncClientModule = (IConcertSyncClientModule*)FModuleManager::Get().GetModule("ConcertSyncClient"))
-	{
-		if (const TSharedPtr<IConcertSyncClient> ConcertSyncClient = ConcertSyncClientModule->GetClient(TEXT("MultiUser")))
-		{
-			const IConcertClientRef ConcertClient = ConcertSyncClient->GetConcertClient();
-			ConcertClient->StartDiscovery();
-		}
-	}
 }
 
 void FDisplayClusterLaunchEditorModule::ShutdownModule()
@@ -96,7 +88,25 @@ void FDisplayClusterLaunchEditorModule::ShutdownModule()
 	{
 		SettingsModule.UnregisterSettings("Project", "Plugins", "nDisplay Launch");
 	}
+	StopMultiUserDiscovery();
+}
 
+void FDisplayClusterLaunchEditorModule::EnsureMultiUserDiscovery()
+{
+	if (TSharedPtr<IConcertSyncClient> ConcertSyncClient = IConcertSyncClientModule::Get().GetClient(TEXT("MultiUser")))
+	{
+		if (!bDidStartDiscovery)
+		{
+			const IConcertClientRef ConcertClient = ConcertSyncClient->GetConcertClient();
+			ConcertClient->StartDiscovery();
+			bDidStartDiscovery = true;
+		}
+	}
+	UE_CLOG(!bDidStartDiscovery, LogDisplayClusterLaunchEditor, Error, TEXT("Unable to start Multi-user discovery process."));
+}
+
+void FDisplayClusterLaunchEditorModule::StopMultiUserDiscovery()
+{
 	if (IConcertSyncClientModule* ConcertSyncClientModule = (IConcertSyncClientModule*)FModuleManager::Get().GetModule("ConcertSyncClient"))
 	{
 		if (const TSharedPtr<IConcertSyncClient> ConcertSyncClient = ConcertSyncClientModule->GetClient(TEXT("MultiUser")))
@@ -105,11 +115,12 @@ void FDisplayClusterLaunchEditorModule::ShutdownModule()
 
 			// Concert may close all existing discovery requests so we have to check to see if we are still have discovery enabled before
 			// attempting to stop.
-			if (ConcertClient->IsDiscoveryEnabled())
+			if (bDidStartDiscovery && ConcertClient->IsDiscoveryEnabled())
 			{
 				ConcertClient->StopDiscovery();
 			}
 			ConcertClient->OnKnownServersUpdated().RemoveAll(this);
+			bDidStartDiscovery = false;
 		}
 	}
 }
@@ -215,10 +226,10 @@ void FDisplayClusterLaunchEditorModule::FindOrLaunchConcertServer()
 	// Shutdown existing server no matter what because we need to hook into OnServersAssumedReady
 	IMultiUserClientModule& MultiUserClientModule = IMultiUserClientModule::Get();
 
+	TArray<FConcertServerInfo> Servers = ConcertClient->GetKnownServers();
 	if (FPlatformProcess::IsProcRunning(ServerTrackingData.MultiUserServerHandle))
 	{
 		// Try to reuse last server
-		TArray<FConcertServerInfo> Servers = ConcertClient->GetKnownServers();
 		FConcertServerInfo* Server = Algo::FindByPredicate(Servers, [Name = GetConcertServerName()](const FConcertServerInfo& Server)
 			{
 				return Server.ServerName == Name;
@@ -236,18 +247,26 @@ void FDisplayClusterLaunchEditorModule::FindOrLaunchConcertServer()
 				);
 		}
 	}
-	else if (MultiUserClientModule.IsConcertServerRunning() || bIsConnectedToSession)
+	else if (Servers.Num() > 0 && (MultiUserClientModule.IsConcertServerRunning() || bIsConnectedToSession))
 	{
 		ConcertServerRequestStatus = EConcertServerRequestStatus::ReuseExisting;
 		OnServersAssumedReady();
 	}
 	else
 	{
+		if (Servers.Num() == 0)
+		{
+			ConcertServerRequestStatus = EConcertServerRequestStatus::ReuseExisting;
+		}
+
 		// We are waiting for servers to register with Concert.
 		ConcertClient->OnKnownServersUpdated().AddRaw(
 			this, &FDisplayClusterLaunchEditorModule::OnServersAssumedReady
 			);
-		LaunchConcertServer();
+		if (!MultiUserClientModule.IsConcertServerRunning())
+		{
+			LaunchConcertServer();
+		}
 	}
 }
 
@@ -341,7 +360,8 @@ void FDisplayClusterLaunchEditorModule::ConnectToSession()
 		}
 		
 		const UConcertClientConfig* CurrentConfig = ConcertClient->GetConfiguration();
-		UConcertClientConfig* AutoConnectConfig = DuplicateObject(CurrentConfig, GetTransientPackage(), CurrentConfig->GetFName());
+		UConcertClientConfig* AutoConnectConfig = DuplicateObject(CurrentConfig, GetTransientPackage());
+
 		AutoConnectConfig->bAutoConnect = true;
 		AutoConnectConfig->bRetryAutoConnectOnError = true;
 		AutoConnectConfig->DefaultServerURL = GetConcertServerName();
@@ -402,6 +422,7 @@ void FDisplayClusterLaunchEditorModule::TryLaunchDisplayClusterProcess()
 	// Create Multi-user params async
 	if (GetConnectToMultiUser())
 	{
+		EnsureMultiUserDiscovery();
 		CachedConcertSessionName.Empty();
 		FindOrLaunchConcertServer();
 	}
diff --git a/Engine/Plugins/Editor/DisplayClusterLaunch/Source/DisplayClusterLaunchEditor/Public/DisplayClusterLaunchEditorModule.h b/Engine/Plugins/Editor/DisplayClusterLaunch/Source/DisplayClusterLaunchEditor/Public/DisplayClusterLaunchEditorModule.h
--- a/Engine/Plugins/Editor/DisplayClusterLaunch/Source/DisplayClusterLaunchEditor/Public/DisplayClusterLaunchEditorModule.h
+++ b/Engine/Plugins/Editor/DisplayClusterLaunch/Source/DisplayClusterLaunchEditor/Public/DisplayClusterLaunchEditorModule.h
@@ -65,6 +65,8 @@ public:
 	void TerminateActiveDisplayClusterProcesses();
 
 private:
+	void EnsureMultiUserDiscovery();
+	void StopMultiUserDiscovery();
 
 	void OnFEngineLoopInitComplete();
 
@@ -134,4 +136,6 @@ private:
 	TArray<FProcHandle> ActiveDisplayClusterProcesses;
 
 	EConcertServerRequestStatus ConcertServerRequestStatus = EConcertServerRequestStatus::None;
+
+	bool bDidStartDiscovery = false;
 };
