nDisplay: Fixed quick launch not finding NDC when Untitled level is saved.

#rnx

diff --git a/Engine/Plugins/Editor/DisplayClusterLaunch/Source/DisplayClusterLaunchEditor/Private/DisplayClusterLaunchEditorModule.cpp b/Engine/Plugins/Editor/DisplayClusterLaunch/Source/DisplayClusterLaunchEditor/Private/DisplayClusterLaunchEditorModule.cpp
--- a/Engine/Plugins/Editor/DisplayClusterLaunch/Source/DisplayClusterLaunchEditor/Private/DisplayClusterLaunchEditorModule.cpp
+++ b/Engine/Plugins/Editor/DisplayClusterLaunch/Source/DisplayClusterLaunchEditor/Private/DisplayClusterLaunchEditorModule.cpp
@@ -395,30 +395,16 @@ void FDisplayClusterLaunchEditorModule::TryLaunchDisplayClusterProcess()
 		return;
 	}
 	
+	RefreshSelectedDisplayActor();
+
 	TArray<TWeakObjectPtr<ADisplayClusterRootActor>> ConfigsInWorld = GetAllDisplayClusterConfigsInWorld();
+
 	if (!DoesCurrentWorldHaveDisplayClusterConfig())
 	{
 		UE_LOG(LogDisplayClusterLaunchEditor, Error, TEXT("%hs: Unable to launch nDisplay because there are no valid nDisplay configurations in the world."), __FUNCTION__);
 		return;
 	}
 	
-	if (!SelectedDisplayClusterConfigActor.IsValid())
-	{
-		for (const TWeakObjectPtr<ADisplayClusterRootActor>& Config : ConfigsInWorld)
-		{
-			if (ADisplayClusterRootActor* ConfigPtr = Config.Get())
-			{
-				TArray<FString> NodeNames;
-				ConfigPtr->GetConfigData()->Cluster->Nodes.GenerateKeyArray(NodeNames);
-				if (NodeNames.Num() > 0)
-				{
-					SetSelectedDisplayClusterConfigActor(ConfigPtr);
-					break;
-				}
-			}
-		}
-	}
-	
 	// Create Multi-user params async
 	if (GetConnectToMultiUser())
 	{
@@ -437,6 +423,7 @@ void FDisplayClusterLaunchEditorModule::LaunchDisplayClusterProcess()
 	UE_LOG(LogDisplayClusterLaunchEditor, Log, TEXT("%hs: Launching nDisplay processes..."), __FUNCTION__);
 
 	// Open a modal to prompt for save, if dirty. Yes = Save & Continue. No = Continue Without Saving. Cancel = Stop Opening Assets.
+
 	UPackage* PackageToSave = nullptr;
 
 	if (UWorld* World = GetCurrentWorld())
@@ -446,6 +433,7 @@ void FDisplayClusterLaunchEditorModule::LaunchDisplayClusterProcess()
 			PackageToSave = Level->GetPackage();
 		}
 	}
+
 	if (PackageToSave)
 	{
 		const FEditorFileUtils::EPromptReturnCode DialogueResponse =
@@ -457,12 +445,15 @@ void FDisplayClusterLaunchEditorModule::LaunchDisplayClusterProcess()
 					LOCTEXT("ConfirmOpenLevelFormat", "Do you want to save the current level?\n\nCancel to abort launch.\n")
 				);
 
-		if (DialogueResponse == FEditorFileUtils::EPromptReturnCode::PR_Cancelled)
+		if (DialogueResponse != FEditorFileUtils::EPromptReturnCode::PR_Success)
 		{
 			return;
 		}
 	}
 	
+	// If the level was renamed (e.g. if it was Untitled), the config actor soft path needs to be updated.
+	RefreshSelectedDisplayActor();
+
 	FString ConcertArguments;
 
 	if (GetConnectToMultiUser())
@@ -944,6 +935,26 @@ void FDisplayClusterLaunchEditorModule::ApplyDisplayClusterConfigOverrides(UDisp
 	ConfigDataCopy->Cluster->Nodes = ActiveNodes;
 }
 
+void FDisplayClusterLaunchEditorModule::RefreshSelectedDisplayActor()
+{
+	if (!SelectedDisplayClusterConfigActor.IsValid() || !Cast<ADisplayClusterRootActor>(SelectedDisplayClusterConfigActor.ResolveObject()))
+	{
+		for (const TWeakObjectPtr<ADisplayClusterRootActor>& Config : GetAllDisplayClusterConfigsInWorld())
+		{
+			if (ADisplayClusterRootActor* ConfigPtr = Config.Get())
+			{
+				TArray<FString> NodeNames;
+				ConfigPtr->GetConfigData()->Cluster->Nodes.GenerateKeyArray(NodeNames);
+				if (NodeNames.Num() > 0)
+				{
+					SetSelectedDisplayClusterConfigActor(ConfigPtr);
+					break;
+				}
+			}
+		}
+	}
+}
+
 void FDisplayClusterLaunchEditorModule::SetSelectedDisplayClusterConfigActor(ADisplayClusterRootActor* SelectedActor)
 {
 	if (SelectedActor)
diff --git a/Engine/Plugins/Editor/DisplayClusterLaunch/Source/DisplayClusterLaunchEditor/Public/DisplayClusterLaunchEditorModule.h b/Engine/Plugins/Editor/DisplayClusterLaunch/Source/DisplayClusterLaunchEditor/Public/DisplayClusterLaunchEditorModule.h
--- a/Engine/Plugins/Editor/DisplayClusterLaunch/Source/DisplayClusterLaunchEditor/Public/DisplayClusterLaunchEditorModule.h
+++ b/Engine/Plugins/Editor/DisplayClusterLaunch/Source/DisplayClusterLaunchEditor/Public/DisplayClusterLaunchEditorModule.h
@@ -97,11 +97,13 @@ private:
 		const UDisplayClusterLaunchEditorProjectSettings* ProjectSettings, FString& ConcatenatedCommandLineArguments, 
 		FString& ConcatenatedConsoleCommands, TSet<FString>& DPCvars, FString& ConcatenatedLogCommands);
 
-
 	TArray<TWeakObjectPtr<ADisplayClusterRootActor>> GetAllDisplayClusterConfigsInWorld();
 	bool DoesCurrentWorldHaveDisplayClusterConfig();
 	void ApplyDisplayClusterConfigOverrides(class UDisplayClusterConfigurationData* ConfigDataCopy);
 
+	/** Updates the selected nDisplay actor if invalid or unresolvable */
+	void RefreshSelectedDisplayActor();
+
 	void SetSelectedDisplayClusterConfigActor(ADisplayClusterRootActor* SelectedActor);
 	void ToggleDisplayClusterConfigActorNodeSelected(FString InNodeName);
 	bool IsDisplayClusterConfigActorNodeSelected(FString InNodeName);
