Add nomcp option to default launch args to display cluster launch.

diff --git a/Engine/Plugins/Editor/DisplayClusterLaunch/Config/BaseDisplayClusterLaunch.ini b/Engine/Plugins/Editor/DisplayClusterLaunch/Config/BaseDisplayClusterLaunch.ini
--- a/Engine/Plugins/Editor/DisplayClusterLaunch/Config/BaseDisplayClusterLaunch.ini
+++ b/Engine/Plugins/Editor/DisplayClusterLaunch/Config/BaseDisplayClusterLaunch.ini
@@ -1,5 +1,5 @@
 [/Script/DisplayClusterLaunchEditor.DisplayClusterLaunchEditorProjectSettings]
 AdditionalConsoleVariables=("p.Chaos.Solver.Deterministic=1","LevelInstance.ForceEditorWorldMode=1")
 AdditionalConsoleCommands=("DisableAllScreenMessages")
-CommandLineArguments=("messaging","dc_cluster","nosplash","fixedseed","NoVerifyGC","noxrstereo","xrtrackingonly","RemoteControlIsHeadless","dx12","dc_dev_mono","unattended","handleensurepercent=0","ini:Engine:[/Script/Engine.Engine]:GameEngine=/Script/DisplayCluster.DisplayClusterGameEngine,[/Script/Engine.Engine]:GameViewportClientClassName=/Script/DisplayCluster.DisplayClusterViewportClient,[/Script/Engine.UserInterfaceSettings]:bAllowHighDPIInGameMode=True,[/Script/Engine.UserInterfaceSettings]:bAllowHighDpiWhenUnattended=True","ini:Game:[/Script/EngineSettings.GeneralProjectSettings]:bUseBorderlessWindow=True")
+CommandLineArguments=("nomcp","messaging","dc_cluster","nosplash","fixedseed","NoVerifyGC","noxrstereo","xrtrackingonly","RemoteControlIsHeadless","dx12","dc_dev_mono","unattended","handleensurepercent=0","ini:Engine:[/Script/Engine.Engine]:GameEngine=/Script/DisplayCluster.DisplayClusterGameEngine,[/Script/Engine.Engine]:GameViewportClientClassName=/Script/DisplayCluster.DisplayClusterViewportClient,[/Script/Engine.UserInterfaceSettings]:bAllowHighDPIInGameMode=True,[/Script/Engine.UserInterfaceSettings]:bAllowHighDpiWhenUnattended=True","ini:Game:[/Script/EngineSettings.GeneralProjectSettings]:bUseBorderlessWindow=True")
 Logging=(Category="LogDisplayClusterRenderSync",VerbosityLevel=Log)
