[nDisplay] Added few more CVars to the DisplayClusterLauncher launch command line

- Copied from SwitchBoard launch command

#rnx

diff --git a/Engine/Plugins/Editor/DisplayClusterLaunch/Source/DisplayClusterLaunchEditor/Private/DisplayClusterLaunchEditorModule.cpp b/Engine/Plugins/Editor/DisplayClusterLaunch/Source/DisplayClusterLaunchEditor/Private/DisplayClusterLaunchEditorModule.cpp
--- a/Engine/Plugins/Editor/DisplayClusterLaunch/Source/DisplayClusterLaunchEditor/Private/DisplayClusterLaunchEditorModule.cpp
+++ b/Engine/Plugins/Editor/DisplayClusterLaunch/Source/DisplayClusterLaunchEditor/Private/DisplayClusterLaunchEditorModule.cpp
@@ -571,6 +571,11 @@ void FDisplayClusterLaunchEditorModule::LaunchDisplayClusterProcess()
 
 		// Log file
 		const FString LogFileName = (ProjectSettings->LogFileName.IsEmpty() ? Node : ProjectSettings->LogFileName) + ".log";
+
+		// Other cvars
+		DPCvars.Add(TEXT("LevelInstance.ForceEditorWorldMode=1"));
+		DPCvars.Add(TEXT("p.Chaos.Solver.Deterministic=1"));
+
 		const FString Params =
 			FString::Printf(
 				TEXT("\"%s\" -game \"%s\" Log=%s %s -ExecCmds=\"%s\" -DPCVars=\"%s\" -LogCmds=\"%s\""),
