Specify the SupportedPlatformTargets in the DisplayClusterLaunch plugin to prevent it from getting included in Mac builds.

#rnx

diff --git a/Engine/Plugins/Editor/DisplayClusterLaunch/DisplayClusterLaunch.uplugin b/Engine/Plugins/Editor/DisplayClusterLaunch/DisplayClusterLaunch.uplugin
--- a/Engine/Plugins/Editor/DisplayClusterLaunch/DisplayClusterLaunch.uplugin
+++ b/Engine/Plugins/Editor/DisplayClusterLaunch/DisplayClusterLaunch.uplugin
@@ -1,5 +1,5 @@
 {
-	"FileVersion": 1,
+	"FileVersion": 3,
 	"Version": 1,
 	"VersionName": "1.0",
 	"FriendlyName": "nDisplay Launch",
@@ -15,6 +15,7 @@
 	"IsBetaVersion": true,
 	"IsExperimentalVersion": false,
 	"Installed": false,
+	"SupportedTargetPlatforms": [ "Win64", "Linux" ],
 	"Plugins":
 	[
 		{
@@ -47,7 +48,7 @@
 			"Name": "DisplayClusterLaunchEditor",
 			"Type": "Editor",
 			"LoadingPhase": "Default",
-			"PlatformAllowlist": [ "Win64", "Linux" ]
+			"PlatformAllowList": ["Win64", "Linux"]
 		}
 	]
 }
